[Skeletal Mesh Editor extension] Experimental morph target editing tools

diff --git a/Engine/Plugins/Experimental/Animation/SkeletalMeshMorphTargetEditingTools/SkeletalMeshMorphTargetEditingTools.uplugin b/Engine/Plugins/Experimental/Animation/SkeletalMeshMorphTargetEditingTools/SkeletalMeshMorphTargetEditingTools.uplugin
new file mode 100644
--- /dev/null
+++ b/Engine/Plugins/Experimental/Animation/SkeletalMeshMorphTargetEditingTools/SkeletalMeshMorphTargetEditingTools.uplugin
@@ -0,0 +1,30 @@
+﻿{
+	"FileVersion": 3,
+	"Version": 1,
+	"VersionName": "1.0",
+	"FriendlyName": "Skeletal Mesh Morph Target Editing Tools",
+	"Description": "Tools to edit morph targets within the skeletal mesh editor.",
+	"Category" : "Animation",
+	"CreatedBy" : "Epic Games, Inc.",
+	"CreatedByURL": "https://epicgames.com/",
+	"DocsURL": "",
+	"MarketplaceURL": "",
+	"CanContainContent": true,
+	"IsBetaVersion": false,
+	"IsExperimentalVersion": true,
+	"Installed": false,
+	"Modules": [
+		{
+			"Name": "SkeletalMeshMorphTargetEditingTools",
+			"Type": "Editor",
+			"LoadingPhase": "Default"
+		}
+	],
+	"Plugins": [
+		{
+			"Name": "SkeletalMeshModelingTools",
+			"Enabled": true,
+			"TargetAllowList": [ "Editor" ]
+		}
+	]
+}
\ No newline at end of file
diff --git a/Engine/Plugins/Experimental/Animation/SkeletalMeshMorphTargetEditingTools/Source/SkeletalMeshMorphTargetEditingTools/Private/MorphTargetEditorTool.cpp b/Engine/Plugins/Experimental/Animation/SkeletalMeshMorphTargetEditingTools/Source/SkeletalMeshMorphTargetEditingTools/Private/MorphTargetEditorTool.cpp
new file mode 100644
--- /dev/null
+++ b/Engine/Plugins/Experimental/Animation/SkeletalMeshMorphTargetEditingTools/Source/SkeletalMeshMorphTargetEditingTools/Private/MorphTargetEditorTool.cpp
@@ -0,0 +1,383 @@
+// Copyright Epic Games, Inc. All Rights Reserved.
+
+#include "MorphTargetEditorTool.h"
+
+#include "ToolBuilderUtil.h"
+#include "InteractiveToolManager.h"
+#include "ModelingToolTargetUtil.h"
+#include "ToolTargetManager.h"
+#include "StaticMeshAttributes.h"
+#include "Engine/SkeletalMesh.h"
+#include "Components/SkeletalMeshComponent.h"
+#include "MeshDescriptionToDynamicMesh.h"
+#include "DynamicMeshToMeshDescription.h"
+#include "TargetInterfaces/PrimitiveComponentBackedTarget.h"
+
+#define LOCTEXT_NAMESPACE "MorphTargetEditorTool"
+
+bool UMorphTargetManagerToolBuilder::CanBuildTool(const FToolBuilderState& SceneState) const
+{
+	return UMultiSelectionMeshEditingToolBuilder::CanBuildTool(SceneState) &&
+		SceneState.TargetManager->CountSelectedAndTargetableWithPredicate(SceneState, GetTargetRequirements(),
+			[](UActorComponent& Component) { return !ToolBuilderUtil::IsVolume(Component); }) >= 1;
+}
+
+UMultiSelectionMeshEditingTool* UMorphTargetManagerToolBuilder::CreateNewTool(const FToolBuilderState& SceneState) const
+{
+	return NewObject<UMorphTargetManagerTool>(SceneState.ToolManager);
+}
+
+void UMorphTargetManagerTool::Setup()
+{
+	UInteractiveTool::Setup();
+
+	if (Targets.Num() != 1)
+	{
+		return;		
+	}
+
+	// Get SkeletalMesh
+	UPrimitiveComponent* TargetComponent = UE::ToolTarget::GetTargetComponent(Targets[0]);
+	USkeletalMeshComponent* Component = Cast<USkeletalMeshComponent>(TargetComponent);
+	USkeletalMesh* SkeletalMesh = Component ? Component->GetSkeletalMeshAsset() : nullptr;
+	if (!SkeletalMesh)
+	{
+		return;
+	}
+
+	// setup modifier
+	Modifier = NewObject<UMorphTargetModifier>(this);
+	Modifier->Init(SkeletalMesh);
+
+	// setup properties
+	Properties = NewObject<UMorphTargetManagerProperties>();
+	Properties->Initialize(this);
+	Properties->RestoreProperties(this);
+	AddToolPropertySource(Properties);
+
+	// Update Lists
+	Modifier->GetMorphTargetNames(Properties->MorphTargetNames);
+}
+
+void UMorphTargetEditorTool::HandleSkeletalMeshModified(const TArray<FName>& BoneNames, const ESkeletalMeshNotifyType InNotifyType)
+{
+	if (GetNotifier().Notifying())
+	{
+		return;
+	}
+
+	// TODO: Add Morph edits to ESkeletalMeshNotifyType
+}
+
+void UMorphTargetEditorTool::CommitResult(UBaseDynamicMeshComponent* Component, bool bModifiedTopology)
+{
+	GetToolManager()->BeginUndoTransaction(LOCTEXT("SculptMeshToolTransactionName", "Sculpt Mesh"));
+	
+	switch (Properties->Operation)
+	{
+		case EMorphTargetEditorOperation::Add:
+			AddMorphTarget();
+			break;
+		case EMorphTargetEditorOperation::Edit:
+			CommitEditMorphTarget();
+			break;
+	}
+
+	GetToolManager()->EndUndoTransaction();
+}
+
+void UMorphTargetManagerTool::Shutdown(EToolShutdownType ShutdownType)
+{
+	if (ShutdownType == EToolShutdownType::Accept)
+	{
+		GetToolManager()->BeginUndoTransaction(LOCTEXT("MorphTargetEditorTool", "Commit Skeleton Editing"));
+		// Modifier->CommitMorphTargetsToSkeletalMesh();
+		GetToolManager()->EndUndoTransaction();
+
+		// to force to refresh the tree
+		if (NeedsNotification())
+		{
+			GetNotifier().Notify({}, ESkeletalMeshNotifyType::BonesAdded);
+		}
+	}
+	
+	Super::Shutdown(ShutdownType);
+
+	// save properties	
+	Properties->SaveProperties(this);	
+}
+
+void UMorphTargetManagerTool::RenameMorphTarget()
+{
+	GetToolManager()->BeginUndoTransaction(LOCTEXT("SculptMeshToolTransactionName", "Sculpt Mesh"));
+
+	// Modify SkeletalMesh
+	const FName& OldName = Properties->RenameOldName;
+	const FName& NewName = Properties->RenameNewName;
+	Modifier->RenameMorphTarget(OldName, NewName);
+
+	GetToolManager()->EndUndoTransaction();
+}
+
+void UMorphTargetManagerTool::RemoveMorphTarget()
+{
+	GetToolManager()->BeginUndoTransaction(LOCTEXT("SculptMeshToolTransactionName", "Sculpt Mesh"));
+
+	// Modify SkeletalMesh
+	const FName& Name = Properties->RemoveName;
+	Modifier->RemoveMorphTarget(Name);
+	
+	GetToolManager()->EndUndoTransaction();
+}
+
+void UMorphTargetManagerTool::RemoveAllMorphTargets()
+{
+	GetToolManager()->BeginUndoTransaction(LOCTEXT("SculptMeshToolTransactionName", "Sculpt Mesh"));
+
+	// Modify SkeletalMesh
+	Modifier->RemoveAllMorphTargets();
+
+	GetToolManager()->EndUndoTransaction();
+}
+
+void UMorphTargetManagerTool::OnTick(float DeltaTime)
+{
+	switch (PendingAction)
+	{
+		case EMorphTargetManagerToolActions::Rename:
+			RenameMorphTarget();
+			break;
+		case EMorphTargetManagerToolActions::Remove:
+			RemoveMorphTarget();
+			break;
+		case EMorphTargetManagerToolActions::RemoveAll:
+			RemoveAllMorphTargets();
+			break;
+	}	
+}
+
+void UMorphTargetManagerTool::HandleSkeletalMeshModified(const TArray<FName>& InBoneNames, const ESkeletalMeshNotifyType InNotifyType)
+{
+	if (GetNotifier().Notifying())
+	{
+		return;
+	}
+}
+
+
+void UMorphTargetManagerTool::RequestAction(const EMorphTargetManagerToolActions& ActionType)
+{
+	if (PendingAction == EMorphTargetManagerToolActions::NoAction)
+	{
+		PendingAction = ActionType;
+	}
+}
+
+void UMorphTargetManagerProperties::Initialize(UMorphTargetManagerTool* ParentToolIn)
+{
+	ParentTool = ParentToolIn;
+}
+
+void UMorphTargetManagerProperties::PostAction(const EMorphTargetManagerToolActions& Action)
+{
+	if (ParentTool.IsValid() && Cast<UMorphTargetManagerTool>(ParentTool))
+	{
+		Cast<UMorphTargetManagerTool>(ParentTool)->RequestAction(Action);
+	}
+}
+
+TArray<FName> UMorphTargetManagerProperties::GetMorphTargetNames()
+{
+	return MorphTargetNames;
+}
+
+UMeshSurfacePointTool* UMorphTargetEditorToolBuilder::CreateNewTool(const FToolBuilderState& SceneState) const
+{
+	UMorphTargetEditorTool* MorphTargetEditorTool = NewObject<UMorphTargetEditorTool>(SceneState.ToolManager);
+	MorphTargetEditorTool->SetWorld(SceneState.World);
+	return MorphTargetEditorTool;
+}
+
+void UMorphTargetEditorTool::Setup()
+{
+	// Add MorphTarget Actions
+	Properties = NewObject<UMorphTargetEditorProperties>();
+	Properties->Initialize(this);
+	Properties->RestoreProperties(this);
+	AddToolPropertySource(Properties);
+
+	// Setup Vertex Sculpt Tool
+	UMeshVertexSculptTool::Setup();
+	
+	// Set Tool Options
+	SetToolDisplayName(LOCTEXT("MorphTargetEditorTool", "Sculpt Morph Target"));
+	SetActiveBrushType(static_cast<uint8>(EMeshVertexSculptBrushType::Move));
+
+
+	// Get SkeletalMesh Component
+	const IPrimitiveComponentBackedTarget* TargetComponent = Cast<IPrimitiveComponentBackedTarget>(Target);
+	USkeletalMeshComponent* Component = Cast<USkeletalMeshComponent>(TargetComponent->GetOwnerComponent());
+	USkeletalMesh* SkeletalMesh = Component->GetSkeletalMeshAsset();
+
+	// Create Modifier
+	Modifier = NewObject<UMorphTargetModifier>(this);
+	Modifier->Init(SkeletalMesh);
+
+	// Update Lists
+	Modifier->GetMorphTargetNames(Properties->MorphTargetNames);
+
+}
+
+void UMorphTargetEditorTool::Shutdown(EToolShutdownType ShutdownType)
+{
+	if (ShutdownType == EToolShutdownType::Accept)
+	{
+		// to force to refresh the tree
+		if (NeedsNotification())
+		{
+			GetNotifier().Notify({}, ESkeletalMeshNotifyType::BonesAdded);
+		}
+	}
+
+	Super::Shutdown(ShutdownType);
+}
+
+void UMorphTargetEditorTool::AddMorphTarget()
+{
+	// Check if Add Operation
+	if (Properties->Operation != EMorphTargetEditorOperation::Add)
+	{
+		return;
+	}
+
+	// Check if target already exists
+	const FName& MorphTargetName = Properties->NewMorphTargetName;
+	if (!Modifier->FindMorphTarget(MorphTargetName))
+	{
+		// Get ScultMesh
+		const FDynamicMesh3* SculptMesh = GetSculptMesh();
+		
+		// Convert to MeshkDescription
+		FMeshDescription SculptMeshDescription;
+		FStaticMeshAttributes Attributes(SculptMeshDescription);
+		Attributes.Register();
+
+		FDynamicMeshToMeshDescription Converter;
+		Converter.Convert(SculptMesh, SculptMeshDescription, false);
+
+		// Edit SkeletalMesh
+		Modifier->AddMorphTarget(MorphTargetName, SculptMeshDescription);
+	}
+}
+
+
+void UMorphTargetEditorTool::SetupEditMorphTarget()
+{
+	// Check if Edit Operation
+	if (Properties->Operation != EMorphTargetEditorOperation::Edit)
+	{
+		return;
+	}
+
+	// Check if target already exists
+	const FName& MorphTargetName = Properties->EditMorphTargetName;
+	const float MorphTargetWeight = FMath::Clamp(Properties->EditMorphTargetWeight, 0.f, 1.f);
+
+	if (MorphTargetWeight > UE_KINDA_SMALL_NUMBER &&
+		Modifier->FindMorphTarget(MorphTargetName))
+	{
+		USkeletalMeshComponent* SkeletalMeshComponent = CastChecked<USkeletalMeshComponent>(UE::ToolTarget::GetTargetComponent(Target));
+		static FTransform T;
+		T = SkeletalMeshComponent->GetRefPoseTransform(1);
+		// Copy Target as MeshDescription
+		FMeshDescription UpdatedTarget = UE::ToolTarget::GetMeshDescriptionCopy(Target);
+		
+		if (Modifier->ApplyMorphTarget(MorphTargetName, UpdatedTarget, MorphTargetWeight))
+		{	
+			BeginChange();
+			{
+				FDynamicMesh3* SculptMesh = GetSculptMesh();
+				SculptMesh->Clear();
+
+				// Get MeshDescription as FDynamicMesh3
+				FDynamicMesh3 TargetDescriptionCopy;
+				FMeshDescriptionToDynamicMesh Converter;
+				Converter.Convert(&UpdatedTarget, *SculptMesh, false/*bCopyTangents*/);
+			}
+			EndChange();
+
+
+			GetSculptMeshComponent()->NotifyMeshUpdated();
+		}
+	}
+}
+
+void UMorphTargetEditorTool::CommitEditMorphTarget()
+{
+	// Check if Edit Operation
+	if (Properties->Operation != EMorphTargetEditorOperation::Edit)
+	{
+		return;
+	}
+
+	// Check if target exists
+	const FName& MorphTargetName = Properties->EditMorphTargetName;
+	if (Modifier->FindMorphTarget(MorphTargetName))
+	{
+		// Get ScultMesh
+		const FDynamicMesh3* SculptMesh = GetSculptMesh();
+
+		// Convert Sculpt to MeshDescription
+		FMeshDescription SculptMeshDescription;
+		FStaticMeshAttributes Attributes(SculptMeshDescription);
+		Attributes.Register();
+
+		FDynamicMeshToMeshDescription Converter;
+		Converter.Convert(SculptMesh, SculptMeshDescription, false);
+
+		// Edit SkeletalMesh
+		Modifier->UpdateMorphTarget(MorphTargetName, SculptMeshDescription);
+	}
+}
+
+void UMorphTargetEditorTool::SetMorphTargetName(FName InName)
+{
+	Properties->Operation = EMorphTargetEditorOperation::Edit;
+	Properties->EditMorphTargetName = InName;
+	FPropertyChangedEvent Event(UMorphTargetEditorProperties::StaticClass()->FindPropertyByName(GET_MEMBER_NAME_CHECKED(UMorphTargetEditorProperties, EditMorphTargetName)));
+	Properties->PostEditChangeProperty(Event);
+}
+
+void UMorphTargetEditorProperties::Initialize(UMorphTargetEditorTool* InParentTool)
+{
+	ParentTool = InParentTool;
+}
+
+TArray<FName> UMorphTargetEditorProperties::GetMorphTargetNames()
+{
+	if (!MorphTargetNames.IsEmpty())
+	{
+		TArray<FName> NoneAndNames = { NAME_None };
+		NoneAndNames.Append(MorphTargetNames);
+		return NoneAndNames;
+	}
+
+	return MorphTargetNames;
+}
+
+void UMorphTargetEditorProperties::PostEditChangeProperty(FPropertyChangedEvent& Event)
+{
+	Super::PostEditChangeProperty(Event);
+
+	const FName PropertyName = (Event.Property != nullptr) ? Event.Property->GetFName() : NAME_None;
+	
+	if (PropertyName == GET_MEMBER_NAME_CHECKED(UMorphTargetEditorProperties, EditMorphTargetName) ||
+		PropertyName == GET_MEMBER_NAME_CHECKED(UMorphTargetEditorProperties, EditMorphTargetWeight))
+	{
+		// Edit
+		ParentTool->SetupEditMorphTarget();
+	}
+}
+
+#undef LOCTEXT_NAMESPACE
+
diff --git a/Engine/Plugins/Experimental/Animation/SkeletalMeshMorphTargetEditingTools/Source/SkeletalMeshMorphTargetEditingTools/Private/MorphTargetModifer.cpp b/Engine/Plugins/Experimental/Animation/SkeletalMeshMorphTargetEditingTools/Source/SkeletalMeshMorphTargetEditingTools/Private/MorphTargetModifer.cpp
new file mode 100644
--- /dev/null
+++ b/Engine/Plugins/Experimental/Animation/SkeletalMeshMorphTargetEditingTools/Source/SkeletalMeshMorphTargetEditingTools/Private/MorphTargetModifer.cpp
@@ -0,0 +1,442 @@
+// Copyright Epic Games, Inc. All Rights Reserved.
+
+
+#include "MorphTargetModifier.h"
+
+#include "MeshDescription.h"
+#include "MeshUtilities.h"
+#include "Animation/MorphTarget.h"
+#include "SkeletalMeshAttributes.h"
+#include "Modules/ModuleManager.h"
+#include "Rendering/SkeletalMeshModel.h"
+#include "Engine/SkeletalMesh.h"
+#include "MeshDescription.h"
+
+bool UMorphTargetModifier::Init(USkeletalMesh* InSkeletalMesh, const int32 InLODIndex)
+{
+	LODIndex = 0;
+	SkeletalMesh = nullptr;
+	MeshDescription.Reset();
+
+#if WITH_EDITORONLY_DATA
+	// validate supplied skeletal mesh exists
+	if (!InSkeletalMesh)
+	{
+		UE_LOG(LogAnimation, Error, TEXT("MorphTarget Modifier: No skeletal mesh supplied to load."));
+		return false;
+	}
+
+	// verify user is not trying to modify one of the core engine assets
+	if (InSkeletalMesh->GetPathName().StartsWith(TEXT("/Engine/")))
+	{
+		UE_LOG(LogAnimation, Error, TEXT("MorphTarget Modifier: Cannot modify built-in engine asset."));
+		return false;
+	}
+	
+	// check LODIndex
+	if (!InSkeletalMesh->IsValidLODIndex(InLODIndex))
+	{
+		UE_LOG(LogAnimation, Error, TEXT("MorphTarget Modifier: Invalid LODIndex."));
+		return false;
+	}
+
+	// store pointer to mesh and instantiate a mesh description for commiting changes
+	SkeletalMesh = InSkeletalMesh;
+	LODIndex = InLODIndex;
+
+	// store mesh description to edit
+	MeshDescription = MakeUnique<FMeshDescription>();
+	SkeletalMesh->CloneMeshDescription(LODIndex, *MeshDescription);
+	if (MeshDescription->IsEmpty())
+	{
+		UE_LOG(LogAnimation, Error, TEXT("MorphTarget Modifier: mesh description is emtpy."));
+		return false;
+	}
+
+	return true;
+
+#else
+	ensureMsgf(false, TEXT("MorphTarget Modifier is an editor only feature."));
+	return false;
+#endif
+}
+
+bool UMorphTargetModifier::GetMorphTargetNames(TArray<FName>& OutNames) const
+{	
+	check(SkeletalMesh);
+
+	OutNames.Reset();
+	for (const UMorphTarget* MorphTarget : SkeletalMesh->GetMorphTargets())
+	{
+		OutNames.Add(MorphTarget->GetFName());
+	}
+	
+	return false;
+}
+
+bool UMorphTargetModifier::GetMorphTargets(TArray<UMorphTarget*>& OutMorphTargets) const
+{
+	check(SkeletalMesh);
+
+	OutMorphTargets = SkeletalMesh->GetMorphTargets();
+
+	return false;
+}
+
+UMorphTarget* UMorphTargetModifier::FindMorphTarget(const FName& Name) const
+{
+	check(SkeletalMesh);
+	if (Name.IsNone() || !Name.IsValid())
+	{
+		return nullptr;
+	}
+
+	int32 Index;
+	return SkeletalMesh->FindMorphTargetAndIndex(Name, Index);
+}
+
+bool UMorphTargetModifier::AddMorphTarget(const FName& Name, const FMeshDescription& TargetMesh) const
+{
+	// Check Name
+	if (Name.IsNone() || !Name.IsValid())
+	{
+		return false;
+	}
+	// Check if already has morph with given name
+	if (FindMorphTarget(Name))
+	{
+		return false;
+	}
+	
+	// Compute Deltas
+	const TArray<FMorphTargetDelta> Deltas = ComputeDeltas(SkeletalMesh, LODIndex, TargetMesh, *MeshDescription);
+
+	// Allocate New MorphTarget
+	UMorphTarget* MorphTarget = NewObject<UMorphTarget>(SkeletalMesh, Name);
+
+	// Get Source MeshModel
+	const FSkeletalMeshModel* MeshModel = SkeletalMesh->GetImportedModel();
+	const FSkeletalMeshLODModel& MeshLODModel = MeshModel->LODModels[LODIndex];
+	const TArray<FSkelMeshSection>& Sections = MeshLODModel.Sections;
+
+	// Populate Deltas
+	MorphTarget->PopulateDeltas(Deltas, LODIndex, Sections, false /*bCompareNormals*/, false/*bGeneratedByReductionSetting*/);
+	
+	// Skip if not deltas found.
+	if (!MorphTarget->HasValidData())
+	{
+		MorphTarget->ConditionalBeginDestroy();
+		return false;
+	}
+
+	// Register Morph Target
+	if (SkeletalMesh->RegisterMorphTarget(MorphTarget, false /*bInvalidateRenderData*/))
+	{
+		// PostEditChange will be called after going out of scope.
+		// This will run the following chain: ImportedData -> SourceModel -> RenderData
+		{
+			FScopedSkeletalMeshPostEditChange ScopedPostEditChange(SkeletalMesh);
+
+			// Update MorphTargetIndexMap
+			SkeletalMesh->InitMorphTargets();
+
+			// CreateImportDataFromLODModel will create the BulkData from LODModel, but only if there is no existing data.
+			// so we have to remove it first.
+			SkeletalMesh->ClearMeshDescriptionAndBulkData(LODIndex);
+
+			const IMeshUtilities& MeshUtilities = FModuleManager::Get().LoadModuleChecked<IMeshUtilities>("MeshUtilities");
+			MeshUtilities.CreateImportDataFromLODModel(SkeletalMesh);
+		}
+
+		SkeletalMesh->MarkPackageDirty();
+
+		return true;
+	}
+
+	return false;
+}
+
+bool UMorphTargetModifier::UpdateMorphTarget(const FName& Name, const FMeshDescription& TargetMesh) const
+{
+	// Check if already has morph with given name
+	UMorphTarget* ExistingMorphTarget = FindMorphTarget(Name);
+	if (!ExistingMorphTarget)
+	{
+		return false;
+	}
+	
+	// Compute Deltas
+	const TArray<FMorphTargetDelta> Deltas = ComputeDeltas(SkeletalMesh, LODIndex, TargetMesh, *MeshDescription);
+
+	// Get Source MeshModel
+	const FSkeletalMeshModel* MeshModel = SkeletalMesh->GetImportedModel();
+	const FSkeletalMeshLODModel& MeshLODModel = MeshModel->LODModels[LODIndex];
+	const TArray<FSkelMeshSection>& Sections = MeshLODModel.Sections;
+
+	{
+		FScopedSkeletalMeshPostEditChange ScopedPostEditChange(SkeletalMesh);
+		
+		ExistingMorphTarget->Modify();
+		// Populate Deltas
+		ExistingMorphTarget->PopulateDeltas(Deltas, LODIndex, Sections, false /*bCompareNormals*/, false/*bGeneratedByReductionSetting*/);
+	
+		// Skip if not deltas found.
+		// if (!ExistingMorphTarget->HasValidData())
+		// {
+		// 	check(false);
+		// 	//TODO: do something about it
+		// 	return false;
+		// }
+		// PostEditChange will be called after going out of scope.
+		// This will run the following chain: ImportedData -> SourceModel -> RenderData
+		{
+			// CreateImportDataFromLODModel will create the BulkData from LODModel, but only if there is no existing data.
+			// so we have to remove it first.
+			SkeletalMesh->ClearMeshDescriptionAndBulkData(LODIndex);
+
+			const IMeshUtilities& MeshUtilities = FModuleManager::Get().LoadModuleChecked<IMeshUtilities>("MeshUtilities");
+			MeshUtilities.CreateImportDataFromLODModel(SkeletalMesh);
+		}
+
+		SkeletalMesh->MarkPackageDirty();
+
+		return true;
+	}
+	
+	return true;
+}
+
+bool UMorphTargetModifier::ApplyMorphTarget(const FName& Name, FMeshDescription& OutMesh, const float Weight) const
+{
+	check(SkeletalMesh);
+
+	// Find MorphTarget
+	if (UMorphTarget* MorphTarget = FindMorphTarget(Name))
+	{
+		// Make sure we have data for this LOD
+		if (MorphTarget->HasDataForLOD(LODIndex))
+		{
+			const FMorphTargetLODModel& MorphModel = MorphTarget->GetMorphLODModels()[LODIndex];
+			const TArray<FMorphTargetDelta>& MorphTargetDeltas = MorphModel.Vertices;
+
+			// Get MeshDescription Vertices	
+			TVertexAttributesRef<FVector3f> OutMeshVertexPositions = OutMesh.GetVertexPositions();
+			const int32 NumElements = OutMeshVertexPositions.GetNumElements();
+
+			const TArray<int32>& RenderToImportedMap = SkeletalMesh->GetImportedModel()->LODModels[LODIndex].MeshToImportVertexMap;
+			
+			for (const FMorphTargetDelta& MorphTargetDelta : MorphTargetDeltas)
+			{
+				// Check if valid Index
+				const int32 VertexIndex = RenderToImportedMap[MorphTargetDelta.SourceIdx];
+					
+				if (VertexIndex >= 0 && VertexIndex < NumElements)
+				{
+					OutMeshVertexPositions[VertexIndex] += MorphTargetDelta.PositionDelta * Weight;
+				}
+			}
+
+			return true;
+		}
+	}
+
+	return false;
+}
+
+//bool UMorphTargetModifier::ApplyMorphTarget(const FName& Name, UE::Geometry::FDynamicMesh3* OutMesh, const float Weight) const
+//{
+//	check(SkeletalMesh);
+//
+//	// Find MorphTarget
+//	if (UMorphTarget* MorphTarget = FindMorphTarget(Name))
+//	{
+//		// Make sure we have data for this LOD
+//		if (MorphTarget->HasDataForLOD(LODIndex))
+//		{
+//			const FMorphTargetLODModel& MorphModel = MorphTarget->GetMorphLODModels()[LODIndex];
+//			const TArray<FMorphTargetDelta>& MorphTargetDeltas = MorphModel.Vertices;
+//
+//			//TVertexAttributesRef<FVector3f> OutMeshVertexPositions = OutMesh.GetVertexPositions();
+//			//const int32 NumElements = OutMeshVertexPositions.GetNumElements();
+//			const int32 NumElements = OutMesh->VertexCount();
+//
+//			for (const FMorphTargetDelta& MorphTargetDelta : MorphTargetDeltas)
+//			{
+//				// Check if valid Index
+//				const int32 VertexIndex = MorphTargetDelta.SourceIdx;
+//				if (VertexIndex >= 0 && VertexIndex < NumElements)
+//				{
+//					const FVector3d Position = OutMesh->GetVertex(VertexIndex);
+//					OutMesh->SetVertex(VertexIndex, Position + FVector3d(MorphTargetDelta.PositionDelta));
+//				}
+//			}
+//
+//			return true;
+//		}
+//	}
+//
+//	return false;
+//}
+
+
+//bool UMorphTargetModifier::AddMorphTarget(const FName& Name, const USkeletalMesh* Target) const
+//{
+//
+//	return false;
+//}
+
+
+bool UMorphTargetModifier::RemoveMorphTarget(const FName& Name) const
+{
+	check(SkeletalMesh);
+
+	// Find MorphTarget
+	if (UMorphTarget* MorphTarget = FindMorphTarget(Name))
+	{
+		// PostEditChange will be called after going out of scope.
+		// This will run the following chain: ImportedData -> SourceModel -> RenderData
+		{
+			FScopedSkeletalMeshPostEditChange ScopedPostEditChange(SkeletalMesh);
+
+			// Unregister MorphTargets
+			// This needs to be under the ScopedPostEdit to avoid triggerring a build that will recreate the morph targets stored in the LODImportData.
+			SkeletalMesh->UnregisterMorphTarget(MorphTarget);
+
+			// CreateImportDataFromLODModel will create the BulkData from LODModel, but only if there is no existing data.
+			// so we have to remove it first.
+			for (int32 InLODIndex = 0; InLODIndex < SkeletalMesh->GetLODNum(); ++InLODIndex)
+			{
+				SkeletalMesh->ClearMeshDescriptionAndBulkData(InLODIndex);
+			}
+
+			const IMeshUtilities& MeshUtilities = FModuleManager::Get().LoadModuleChecked<IMeshUtilities>("MeshUtilities");
+			MeshUtilities.CreateImportDataFromLODModel(SkeletalMesh);
+		}
+
+		SkeletalMesh->MarkPackageDirty();
+	}
+
+	return true;
+}
+
+
+bool UMorphTargetModifier::RemoveAllMorphTargets() const
+{
+	check(SkeletalMesh);
+
+	// Skip of already empty.
+	if (SkeletalMesh->GetMorphTargets().IsEmpty())
+	{
+		return false;
+	}
+
+	// PostEditChange will be called after going out of scope.
+	// This will run the following chain: ImportedData -> SourceModel -> RenderData
+	{
+		FScopedSkeletalMeshPostEditChange ScopedPostEditChange(SkeletalMesh);
+
+		// Remove All Targets
+		// This needs to be under the ScopedPostEdit to avoid triggerring a build that will recreate the morph targets stored in the LODImportData.
+		SkeletalMesh->UnregisterAllMorphTarget();
+
+		// CreateImportDataFromLODModel will create the BulkData from LODModel, but only if there is no existing ImportedData.
+		// so we have to remove it first.
+		for (int32 InLODIndex = 0; InLODIndex < SkeletalMesh->GetLODNum(); ++InLODIndex)
+		{
+			SkeletalMesh->ClearMeshDescriptionAndBulkData(InLODIndex);
+		}
+
+		const IMeshUtilities& MeshUtilities = FModuleManager::Get().LoadModuleChecked<IMeshUtilities>("MeshUtilities");
+		MeshUtilities.CreateImportDataFromLODModel(SkeletalMesh);
+	}
+
+	SkeletalMesh->MarkPackageDirty();
+
+	return true;
+}
+
+bool UMorphTargetModifier::RenameMorphTarget(const FName& InOldMorphTargetName, const FName& InNewMorphTargetName)
+{
+	return false;
+}
+
+bool UMorphTargetModifier::CommitMorphTargetsToSkeletalMesh()
+{
+	if (!SkeletalMesh || !MeshDescription)
+	{
+		UE_LOG(LogAnimation, Error, TEXT("MorphTarget Modifier: No mesh loaded. Cannot apply MorphTarget edits."));
+		return false;
+	}
+
+#if WITH_EDITORONLY_DATA	
+	//// update skeletal mesh
+	//FlushRenderingCommands();
+	//
+	//SkeletalMesh->SetFlags(RF_Transactional);
+	//SkeletalMesh->Modify();
+
+	//// update the ref skeleton
+	//SkeletalMesh->SetRefSkeleton(*ReferenceSkeleton);
+	//SkeletalMesh->GetRefBasesInvMatrix().Reset();
+	//SkeletalMesh->CalculateInvRefMatrices();
+	//
+	//// update skeletal mesh LOD (cf. USkeletalMesh::CommitMeshDescription)
+	//SkeletalMesh->CommitMeshDescription(USkeletonModifierLocals::LODIndex, *MeshDescription);
+
+	//// update skeleton
+	//USkeleton* Skeleton = SkeletalMesh->GetSkeleton();
+	//Skeleton->Modify();
+	//if (Skeleton->RecreateBoneTree(SkeletalMesh))
+	//{
+	//	Skeleton->MarkPackageDirty();	
+	//}
+	//
+	//// must be done once the skeleton is up to date
+	//SkeletalMesh->PostEditChange();
+
+	return true;
+#else
+	ensureMsgf(false, TEXT("MorphTarget Modifier is an editor only feature."));
+	return false;
+#endif
+}
+
+bool UMorphTargetModifier::SetMorphTarget(const FName& Name, const FMeshDescription& Target) const
+{
+
+	return false;
+}
+
+TArray<FMorphTargetDelta> UMorphTargetModifier::ComputeDeltas(USkeletalMesh* InSkeletalMesh, int32 InLodIndex, const FMeshDescription& TargetMesh, const FMeshDescription& BaseMesh)
+{
+	const FVertexArray& TargetVertexIDs = TargetMesh.Vertices();
+	const FVertexArray& BaseVertexIDs = BaseMesh.Vertices();
+	
+	const int32 NumVertices = BaseVertexIDs.Num();
+	check(TargetVertexIDs.Num() == NumVertices)
+
+	//const TVertexAttributesConstRef<FVector3f> TargetVertexPositions = TargetMesh.VertexAttributes().GetAttributesRef<FVector3f>(MeshAttribute::Vertex::Position);
+	//const TVertexAttributesConstRef<FVector3f> BaseVertexPositions = BaseMesh.VertexAttributes().GetAttributesRef<FVector3f>(MeshAttribute::Vertex::Position);
+
+	const TVertexAttributesConstRef<FVector3f> TargetVertexPositions = TargetMesh.GetVertexPositions();
+	const TVertexAttributesConstRef<FVector3f> BaseVertexPositions = BaseMesh.GetVertexPositions();
+
+	const TArray<int32>& RenderToImportedMap = InSkeletalMesh->GetImportedModel()->LODModels[InLodIndex].MeshToImportVertexMap;
+	int32 NumRenderVerts = InSkeletalMesh->GetImportedModel()->LODModels[InLodIndex].NumVertices;
+	
+	TArray<FMorphTargetDelta> Deltas;
+	Deltas.SetNumZeroed(NumRenderVerts);
+	for (int32 RenderIndex = 0; RenderIndex < NumRenderVerts; RenderIndex++)
+	{
+		const FVertexID ImportedVertexID(RenderToImportedMap[RenderIndex]);
+		const FVector3f& TargetVertex = TargetVertexPositions.Get(ImportedVertexID);
+		const FVector3f& BaseVertex = BaseVertexPositions.Get(ImportedVertexID);
+
+		Deltas[RenderIndex].PositionDelta = TargetVertex - BaseVertex;
+		Deltas[RenderIndex].SourceIdx = RenderIndex;
+	}
+
+	return Deltas;
+}
+
+
diff --git a/Engine/Plugins/Experimental/Animation/SkeletalMeshMorphTargetEditingTools/Source/SkeletalMeshMorphTargetEditingTools/Private/SkeletalMeshMorphTargetEditingToolsCommands.cpp b/Engine/Plugins/Experimental/Animation/SkeletalMeshMorphTargetEditingTools/Source/SkeletalMeshMorphTargetEditingTools/Private/SkeletalMeshMorphTargetEditingToolsCommands.cpp
new file mode 100644
--- /dev/null
+++ b/Engine/Plugins/Experimental/Animation/SkeletalMeshMorphTargetEditingTools/Source/SkeletalMeshMorphTargetEditingTools/Private/SkeletalMeshMorphTargetEditingToolsCommands.cpp
@@ -0,0 +1,29 @@
+// Copyright Epic Games, Inc. All Rights Reserved.
+
+#include "SkeletalMeshMorphTargetEditingToolsCommands.h"
+#include "SkeletalMeshMorphTargetEditingToolsStyle.h"
+
+#define LOCTEXT_NAMESPACE "SkeletalMeshMorphTargetEditingToolsCommands"
+
+
+
+FSkeletalMeshMorphTargetEditingToolsCommands::FSkeletalMeshMorphTargetEditingToolsCommands()
+	: TCommands<FSkeletalMeshMorphTargetEditingToolsCommands>(
+		"SkeletalMeshMorphTargetEditingToolsCommands", // Context name for fast lookup
+		NSLOCTEXT("Contexts", "SkeletalMeshMorphTargetEditingToolsCommands", "Skeletal Mesh Morph Target Editing Tools"), // Localized context name for displaying
+		NAME_None, // Parent
+		FSkeletalMeshMorphTargetEditingToolsStyle::Get().GetStyleSetName() // Icon Style Set
+		)
+{
+}
+
+void FSkeletalMeshMorphTargetEditingToolsCommands::RegisterCommands()
+{
+	UI_COMMAND(BeginMorphTargetTool, "Morph", "Edit Morph Targets", EUserInterfaceActionType::ToggleButton, FInputChord());
+	
+	UI_COMMAND(BeginMorphTargetEditorTool, "Add/Edit Morph Target", "MorphTarget Sculpt Tool", EUserInterfaceActionType::ToggleButton, FInputChord());
+	UI_COMMAND(BeginMorphTargetManagerTool, "Manange Morph Targets ", "MorphTarget Manager Tool", EUserInterfaceActionType::ToggleButton, FInputChord());
+}
+
+#undef LOCTEXT_NAMESPACE
+
diff --git a/Engine/Plugins/Experimental/Animation/SkeletalMeshMorphTargetEditingTools/Source/SkeletalMeshMorphTargetEditingTools/Private/SkeletalMeshMorphTargetEditingToolsModule.cpp b/Engine/Plugins/Experimental/Animation/SkeletalMeshMorphTargetEditingTools/Source/SkeletalMeshMorphTargetEditingTools/Private/SkeletalMeshMorphTargetEditingToolsModule.cpp
new file mode 100644
--- /dev/null
+++ b/Engine/Plugins/Experimental/Animation/SkeletalMeshMorphTargetEditingTools/Source/SkeletalMeshMorphTargetEditingTools/Private/SkeletalMeshMorphTargetEditingToolsModule.cpp
@@ -0,0 +1,60 @@
+﻿// Copyright Epic Games, Inc. All Rights Reserved.
+
+#include "SkeletalMeshMorphTargetEditingToolsModule.h"
+
+#include "MorphTargetEditorTool.h"
+#include "SkeletalMeshMorphTargetEditingToolsCommands.h"
+#include "SkeletalMeshMorphTargetEditingToolsStyle.h"
+#include "Features/IModularFeatures.h"
+#include "Framework/Commands/Commands.h"
+
+#define LOCTEXT_NAMESPACE "FSkeletalMeshMorphTargetEditingToolsModule"
+
+void FSkeletalMeshMorphTargetEditingToolsModule::StartupModule()
+{
+	FSkeletalMeshMorphTargetEditingToolsStyle::Register();
+	FSkeletalMeshMorphTargetEditingToolsCommands::Register();
+	IModularFeatures::Get().RegisterModularFeature(ISkeletalMeshModelingModeToolExtension::GetModularFeatureName(), this);	
+
+}
+
+void FSkeletalMeshMorphTargetEditingToolsModule::ShutdownModule()
+{
+	IModularFeatures::Get().UnregisterModularFeature(ISkeletalMeshModelingModeToolExtension::GetModularFeatureName(), this);
+	FSkeletalMeshMorphTargetEditingToolsCommands::Unregister();
+	FSkeletalMeshMorphTargetEditingToolsStyle::Unregister();
+}
+
+
+
+void FSkeletalMeshMorphTargetEditingToolsModule::GetExtensionTools(const FExtensionToolQueryInfo& QueryInfo,
+                                                                   TArray<FExtensionToolDescription>& OutTools)
+{
+	const FSkeletalMeshMorphTargetEditingToolsCommands& Commands = FSkeletalMeshMorphTargetEditingToolsCommands::Get();
+	{
+		FExtensionToolDescription ToolDesc;
+		ToolDesc.ToolName = LOCTEXT("SkeletalMeshMorphTargetEditorTool", "Add/Edit Morph Target");
+		ToolDesc.ToolCommand = Commands.BeginMorphTargetEditorTool;
+		ToolDesc.ToolBuilder = NewObject<UMorphTargetEditorToolBuilder>();
+		OutTools.Add(ToolDesc);
+	}
+
+	{
+		FExtensionToolDescription ToolDesc;
+		ToolDesc.ToolName = LOCTEXT("SkeletalMeshMorphTargetManagerTool", "Manange Morph Targets");
+		ToolDesc.ToolCommand = Commands.BeginMorphTargetManagerTool;
+		ToolDesc.ToolBuilder = NewObject<UMorphTargetManagerToolBuilder>();
+		OutTools.Add(ToolDesc);
+	}	
+}
+
+bool FSkeletalMeshMorphTargetEditingToolsModule::GetExtensionExtendedInfo(FModelingModeExtensionExtendedInfo& InfoOut)
+{
+	InfoOut.ExtensionCommand = FSkeletalMeshMorphTargetEditingToolsCommands::Get().BeginMorphTargetTool;
+	
+	return true;
+}
+
+#undef LOCTEXT_NAMESPACE
+	
+IMPLEMENT_MODULE(FSkeletalMeshMorphTargetEditingToolsModule, SkeletalMeshMorphTargetEditingTools)
\ No newline at end of file
diff --git a/Engine/Plugins/Experimental/Animation/SkeletalMeshMorphTargetEditingTools/Source/SkeletalMeshMorphTargetEditingTools/Private/SkeletalMeshMorphTargetEditingToolsStyle.cpp b/Engine/Plugins/Experimental/Animation/SkeletalMeshMorphTargetEditingTools/Source/SkeletalMeshMorphTargetEditingTools/Private/SkeletalMeshMorphTargetEditingToolsStyle.cpp
new file mode 100644
--- /dev/null
+++ b/Engine/Plugins/Experimental/Animation/SkeletalMeshMorphTargetEditingTools/Source/SkeletalMeshMorphTargetEditingTools/Private/SkeletalMeshMorphTargetEditingToolsStyle.cpp
@@ -0,0 +1,40 @@
+// Copyright Epic Games, Inc. All Rights Reserved.
+
+#include "SkeletalMeshMorphTargetEditingToolsStyle.h"
+
+#include "Styling/SlateStyleMacros.h"
+#include "Styling/SlateStyleRegistry.h"
+
+
+FSkeletalMeshMorphTargetEditingToolsStyle& FSkeletalMeshMorphTargetEditingToolsStyle::Get()
+{
+	static FSkeletalMeshMorphTargetEditingToolsStyle Instance;
+	return Instance;
+}
+
+void FSkeletalMeshMorphTargetEditingToolsStyle::Register()
+{
+	FSlateStyleRegistry::RegisterSlateStyle(Get());
+}
+
+void FSkeletalMeshMorphTargetEditingToolsStyle::Unregister()
+{
+	FSlateStyleRegistry::UnRegisterSlateStyle(Get());
+}
+
+FSkeletalMeshMorphTargetEditingToolsStyle::FSkeletalMeshMorphTargetEditingToolsStyle()
+	: FSlateStyleSet("SkeletalMeshMorphTargetEditingToolsStyle")
+{
+	static const FVector2D IconSize10x10(10.0f, 10.0f);
+	static const FVector2D IconSize16x12(16.0f, 12.0f);
+	static const FVector2D IconSize16x16(16.0f, 16.0f);
+	static const FVector2D IconSize20x20(20.0f, 20.0f);
+	static const FVector2D IconSize64x64(64.0f, 64.0f);
+
+	static const FSlateColor DefaultForeground(FLinearColor(0.72f, 0.72f, 0.72f, 1.f));
+
+	FSlateStyleSet::SetContentRoot(FPaths::EnginePluginsDir() / TEXT("Experimental/Animation/SkeletalMeshMorphTargetEditingTools/Resources"));
+	{
+		Set("SkeletalMeshMorphTargetEditingToolsCommands.BeginMorphTargetTool", new IMAGE_BRUSH("Icon128", IconSize20x20));
+	}
+}
diff --git a/Engine/Plugins/Experimental/Animation/SkeletalMeshMorphTargetEditingTools/Source/SkeletalMeshMorphTargetEditingTools/Public/MorphTargetEditorTool.h b/Engine/Plugins/Experimental/Animation/SkeletalMeshMorphTargetEditingTools/Source/SkeletalMeshMorphTargetEditingTools/Public/MorphTargetEditorTool.h
new file mode 100644
--- /dev/null
+++ b/Engine/Plugins/Experimental/Animation/SkeletalMeshMorphTargetEditingTools/Source/SkeletalMeshMorphTargetEditingTools/Public/MorphTargetEditorTool.h
@@ -0,0 +1,217 @@
+// Copyright Epic Games, Inc. All Rights Reserved.
+
+#pragma once
+
+#include "MorphTargetModifier.h"
+#include "MeshVertexSculptTool.h"
+#include "SkeletalMesh/SkeletalMeshEditionInterface.h"
+#include "BaseTools/MultiSelectionMeshEditingTool.h"
+
+#include "MorphTargetEditorTool.generated.h"
+
+UENUM()
+enum class EMorphTargetEditorOperation
+{
+	Add,  // Adds New MorphTarget
+	Edit, // Edit existing MorphTarget
+};
+
+UENUM()
+enum class EMorphTargetManagerToolActions
+{
+	NoAction,
+	Rename,    // Renames MorphTarget with new Name
+	Remove,    // Removes existing MorphTarget
+	RemoveAll  // Removes all MorphTargets
+};
+/**
+ * MorphTaget Editor Tool Builder
+ */
+UCLASS()
+class SKELETALMESHMORPHTARGETEDITINGTOOLS_API UMorphTargetEditorToolBuilder : public UMeshVertexSculptToolBuilder
+{
+	GENERATED_BODY()
+
+public:
+	virtual UMeshSurfacePointTool* CreateNewTool(const FToolBuilderState& SceneState) const override;
+};
+
+/**
+* MorphTarget Editor Tool
+*/
+UCLASS()
+class SKELETALMESHMORPHTARGETEDITINGTOOLS_API UMorphTargetEditorTool : 
+	public UMeshVertexSculptTool, 
+	public ISkeletalMeshEditingInterface
+{
+	GENERATED_BODY()
+
+public:	
+	// UMeshVertexSculptTool overridess
+	virtual void Setup() override;
+	virtual void Shutdown(EToolShutdownType ShutdownType) override;
+
+public:
+	// modifier functions
+	void AddMorphTarget();
+	void SetupEditMorphTarget();
+	void CommitEditMorphTarget();
+
+	void SetMorphTargetName(FName InName);
+protected:
+
+	virtual void CommitResult(UBaseDynamicMeshComponent* Component, bool bModifiedTopology) override;
+
+	// ISkeletalMeshEditionInterface override
+	// TODO: add ESkeletalMeshNotifyType::AddMorphTarget
+	virtual void HandleSkeletalMeshModified(const TArray<FName>& BoneNames, const ESkeletalMeshNotifyType InNotifyType) override;
+
+	UPROPERTY()
+	TObjectPtr<UMorphTargetEditorProperties> Properties;
+
+	UPROPERTY()
+	TObjectPtr<UMorphTargetModifier> Modifier;
+
+};
+
+// ---------------------------------------------------------------------------------------------------------------------
+
+/**
+* MorphTarget Manager Tool Builder
+*/
+UCLASS()
+class SKELETALMESHMORPHTARGETEDITINGTOOLS_API UMorphTargetManagerToolBuilder : public UMultiSelectionMeshEditingToolBuilder
+{
+	GENERATED_BODY()
+
+public:
+	virtual UMultiSelectionMeshEditingTool* CreateNewTool(const FToolBuilderState& SceneState) const override;
+	virtual bool CanBuildTool(const FToolBuilderState& SceneState) const override;
+};
+
+/**
+ * MorphTarget Manager Tool
+ */
+UCLASS()
+class SKELETALMESHMORPHTARGETEDITINGTOOLS_API UMorphTargetManagerTool :
+	public UMultiSelectionMeshEditingTool,
+	public ISkeletalMeshEditingInterface
+{
+	GENERATED_BODY()
+
+public:
+
+	virtual void Setup() override;
+	virtual void Shutdown(EToolShutdownType ShutdownType) override;
+	virtual void OnTick(float DeltaTime) override;
+
+	virtual bool HasCancel() const override { return false; }
+	virtual bool HasAccept() const override { return false; }
+	virtual bool CanAccept() const override { return false; }
+
+	void RequestAction(const EMorphTargetManagerToolActions& ActionType);
+
+protected:
+
+	EMorphTargetManagerToolActions PendingAction = EMorphTargetManagerToolActions::NoAction;
+
+	// Modifier functions
+	void RenameMorphTarget();
+	void RemoveMorphTarget();
+	void RemoveAllMorphTargets();
+
+	// ISkeletalMeshEditionInterface overrides
+	// TODO: Add MorphTarget Notifiers
+	virtual void HandleSkeletalMeshModified(const TArray<FName>& BoneNames, const ESkeletalMeshNotifyType InNotifyType) override;
+
+	UPROPERTY()
+	TObjectPtr<UMorphTargetModifier> Modifier;
+	
+	UPROPERTY()
+	TObjectPtr<UMorphTargetManagerProperties> Properties;
+};
+
+// ---------------------------------------------------------------------------------------
+//									PROPERTIES
+// ---------------------------------------------------------------------------------------
+
+UCLASS()
+class SKELETALMESHMORPHTARGETEDITINGTOOLS_API UMorphTargetEditorProperties : public UInteractiveToolPropertySet
+{
+	GENERATED_BODY()
+public:
+
+	UPROPERTY(EditAnywhere, Category = "MorphTarget")
+	EMorphTargetEditorOperation Operation;
+
+	UPROPERTY(EditAnywhere, Category = "MorphTarget", meta = (DisplayName = "Name", EditCondition = "Operation == EMorphTargetEditorOperation::Add", EditConditionHides))
+	FName NewMorphTargetName;
+	
+	UPROPERTY(EditAnywhere, Category = "MorphTarget", meta = (GetOptions = GetMorphTargetNames, NoResetToDefault, DisplayName = "Name", EditCondition = "Operation == EMorphTargetEditorOperation::Edit", EditConditionHides))
+	FName EditMorphTargetName;
+
+	UPROPERTY(EditAnywhere, Category = "MorphTarget", meta = (DisplayName = "Weight", EditCondition = "Operation == EMorphTargetEditorOperation::Edit", EditConditionHides, ClampMin = "0.0", ClampMax = "1.0", UIMin = "0.0", UIMax = "1.0"))
+	float EditMorphTargetWeight = 1.f;
+
+	TWeakObjectPtr<UMorphTargetEditorTool> ParentTool;
+
+	void Initialize(UMorphTargetEditorTool* ParentToolIn);
+
+	UFUNCTION()
+	TArray<FName> GetMorphTargetNames();
+
+	UPROPERTY()
+	TArray<FName> MorphTargetNames;
+
+	virtual void PostEditChangeProperty(struct FPropertyChangedEvent& Event) override;
+	
+	UPROPERTY()
+    TObjectPtr<UPreviewMesh> MorphTargetPreviewMesh = nullptr;
+    UPROPERTY()
+    TObjectPtr<UToolTarget> MorphTargetTarget = nullptr;
+};
+
+UCLASS()
+class SKELETALMESHMORPHTARGETEDITINGTOOLS_API UMorphTargetManagerProperties : public UInteractiveToolPropertySet
+{
+	GENERATED_BODY()
+
+public:
+	TWeakObjectPtr<UMorphTargetManagerTool> ParentTool;
+
+	void Initialize(UMorphTargetManagerTool* ParentToolIn);
+	void PostAction(const EMorphTargetManagerToolActions& Action);
+
+	UPROPERTY(EditAnywhere, Category = "RenameMorphTarget", meta = (GetOptions = GetMorphTargetNames, NoResetToDefault, DisplayName="OldName"))
+	FName RenameOldName;
+
+	UPROPERTY(EditAnywhere, Category = "RenameMorphTarget", meta = (DisplayName = "NewName"))
+	FName RenameNewName;
+
+	UPROPERTY(EditAnywhere, Category = "RemoveMorphTarget", meta = (GetOptions = GetMorphTargetNames, NoResetToDefault, DisplayName = "Name"))
+	FName RemoveName;
+
+	UFUNCTION()
+	TArray<FName> GetMorphTargetNames();
+
+	UPROPERTY()
+	TArray<FName> MorphTargetNames;
+
+	UFUNCTION(CallInEditor, Category = "RenameMorphTarget")
+	void RenameMorphTarget()
+	{
+		PostAction(EMorphTargetManagerToolActions::Rename);
+	}
+
+	UFUNCTION(CallInEditor, Category = "RemoveMorphTarget")
+	void RemoveMorphTarget()
+	{
+		PostAction(EMorphTargetManagerToolActions::Remove);
+	}
+
+	UFUNCTION(CallInEditor, Category = "RemoveAllMorphTargets")
+		void RemoveAllMorphTargets()
+	{
+		PostAction(EMorphTargetManagerToolActions::RemoveAll);
+	}
+};
diff --git a/Engine/Plugins/Experimental/Animation/SkeletalMeshMorphTargetEditingTools/Source/SkeletalMeshMorphTargetEditingTools/Public/MorphTargetModifier.h b/Engine/Plugins/Experimental/Animation/SkeletalMeshMorphTargetEditingTools/Source/SkeletalMeshMorphTargetEditingTools/Public/MorphTargetModifier.h
new file mode 100644
--- /dev/null
+++ b/Engine/Plugins/Experimental/Animation/SkeletalMeshMorphTargetEditingTools/Source/SkeletalMeshMorphTargetEditingTools/Public/MorphTargetModifier.h
@@ -0,0 +1,85 @@
+// Copyright Epic Games, Inc. All Rights Reserved.
+
+#pragma once
+
+#include "Math/MathFwd.h"
+#include "Containers/Array.h"
+#include "Templates/UniquePtr.h"
+
+#include "MorphTargetModifier.generated.h"
+
+class USkeletalMesh;
+class UMorphTarget;
+class FName;
+struct FMorphTargetDelta;
+struct FMeshDescription;
+
+UCLASS(BlueprintType, Blueprintable)
+class SKELETALMESHMORPHTARGETEDITINGTOOLS_API UMorphTargetModifier : public UObject
+{
+	GENERATED_BODY()
+	
+public:
+
+	UFUNCTION(BlueprintCallable, Category = "Mesh")
+	bool Init(USkeletalMesh* InSkeletalMesh, const int32 InLODIndex=0);
+
+	/**
+	* Returns MorphTargets in SkeletalMesh
+	*/
+	UFUNCTION(BlueprintCallable, Category = "MorphTarget")
+	UMorphTarget* FindMorphTarget(const FName& Name) const;
+
+	UFUNCTION(BlueprintCallable, Category = "MorphTarget")
+	bool GetMorphTargetNames(TArray<FName>& OutNames) const;
+	UFUNCTION(BlueprintCallable, Category = "MorphTarget")
+	bool GetMorphTargets(TArray<UMorphTarget*>& OutMorphTargets) const;
+
+	/**
+	* Add MorphTargets from MeshDescription
+	*/
+	bool AddMorphTarget(const FName& Name, const FMeshDescription& Target) const;
+	bool UpdateMorphTarget(const FName& Name, const FMeshDescription& Target) const;
+
+	/***
+	* Apply MorphTarget to MeshDescription
+	*/
+	bool ApplyMorphTarget(const FName& Name, FMeshDescription& OutMesh, 
+		const float Weight = 1.0) const;
+
+	/***
+	* Apply MorphTarget to MeshDescription
+	*/
+	//bool ApplyMorphTarget(const FName& Name, UE::Geometry::FDynamicMesh3* OutMesh,
+	//	const float Weight = 1.0) const;
+
+	/**
+	* Remove MorphTargets from SkeletalMesh
+	*/
+	UFUNCTION(BlueprintCallable, Category = "MorphTarget")
+	bool RemoveMorphTarget(const FName& Name) const;
+	bool RemoveAllMorphTargets() const;
+
+	UFUNCTION(BlueprintCallable, Category = "MorphTarget")
+	bool RenameMorphTarget(const FName& InOldName, const FName& InNewName);
+	
+	/**
+	 * Actually applies the skeleton modifications to the skeletal mesh.
+	 * @return true if commit succeeded.
+	 */
+	UFUNCTION(BlueprintCallable, Category = "Mesh")
+	bool CommitMorphTargetsToSkeletalMesh();
+
+private:
+
+	bool SetMorphTarget(const FName& Name, const FMeshDescription& Target) const;
+	static TArray<FMorphTargetDelta> ComputeDeltas(USkeletalMesh* InSkeletalMesh, int32 InLodIndex, const FMeshDescription& TargetMesh, const FMeshDescription& BaseMesh);
+
+	int32 LODIndex = 0;
+
+	UPROPERTY()
+	TObjectPtr<USkeletalMesh> SkeletalMesh = nullptr;
+
+	TUniquePtr<FMeshDescription> MeshDescription;
+};
+
diff --git a/Engine/Plugins/Experimental/Animation/SkeletalMeshMorphTargetEditingTools/Source/SkeletalMeshMorphTargetEditingTools/Public/SkeletalMeshMorphTargetEditingToolsCommands.h b/Engine/Plugins/Experimental/Animation/SkeletalMeshMorphTargetEditingTools/Source/SkeletalMeshMorphTargetEditingTools/Public/SkeletalMeshMorphTargetEditingToolsCommands.h
new file mode 100644
--- /dev/null
+++ b/Engine/Plugins/Experimental/Animation/SkeletalMeshMorphTargetEditingTools/Source/SkeletalMeshMorphTargetEditingTools/Public/SkeletalMeshMorphTargetEditingToolsCommands.h
@@ -0,0 +1,19 @@
+// Copyright Epic Games, Inc. All Rights Reserved.
+
+#pragma once
+
+#include "CoreMinimal.h"
+#include "Framework/Commands/Commands.h"
+
+class FSkeletalMeshMorphTargetEditingToolsCommands : public TCommands<FSkeletalMeshMorphTargetEditingToolsCommands>
+{
+public:
+	FSkeletalMeshMorphTargetEditingToolsCommands();
+
+
+	TSharedPtr<FUICommandInfo> BeginMorphTargetTool;
+	TSharedPtr<FUICommandInfo> BeginMorphTargetEditorTool;
+	TSharedPtr<FUICommandInfo> BeginMorphTargetManagerTool;
+
+	virtual void RegisterCommands() override;
+};
diff --git a/Engine/Plugins/Experimental/Animation/SkeletalMeshMorphTargetEditingTools/Source/SkeletalMeshMorphTargetEditingTools/Public/SkeletalMeshMorphTargetEditingToolsModule.h b/Engine/Plugins/Experimental/Animation/SkeletalMeshMorphTargetEditingTools/Source/SkeletalMeshMorphTargetEditingTools/Public/SkeletalMeshMorphTargetEditingToolsModule.h
new file mode 100644
--- /dev/null
+++ b/Engine/Plugins/Experimental/Animation/SkeletalMeshMorphTargetEditingTools/Source/SkeletalMeshMorphTargetEditingTools/Public/SkeletalMeshMorphTargetEditingToolsModule.h
@@ -0,0 +1,26 @@
+﻿// Copyright Epic Games, Inc. All Rights Reserved.
+
+#pragma once
+
+#include "Modules/ModuleManager.h"
+#include "SkeletalMeshModelingModeToolExtensions.h"
+
+class FUICommandInfo;
+
+class FSkeletalMeshMorphTargetEditingToolsModule : public IModuleInterface, public ISkeletalMeshModelingModeToolExtension 
+{
+public:
+	/** IModuleInterface implementation */
+	virtual void StartupModule() override;
+	virtual void ShutdownModule() override;
+	
+	// ISkeletalMeshModelingModeToolExtension
+	FText GetExtensionName() override { return FText::FromString(TEXT("SkeletalMeshMorphTargetEditingTools"));}
+
+	FText GetToolSectionName() override { return FText::FromString(TEXT("Morph"));};
+
+	void GetExtensionTools(const FExtensionToolQueryInfo& QueryInfo, TArray<FExtensionToolDescription>& OutTools) override;
+	
+	bool GetExtensionExtendedInfo(FModelingModeExtensionExtendedInfo& InfoOut) override;
+
+};
diff --git a/Engine/Plugins/Experimental/Animation/SkeletalMeshMorphTargetEditingTools/Source/SkeletalMeshMorphTargetEditingTools/Public/SkeletalMeshMorphTargetEditingToolsStyle.h b/Engine/Plugins/Experimental/Animation/SkeletalMeshMorphTargetEditingTools/Source/SkeletalMeshMorphTargetEditingTools/Public/SkeletalMeshMorphTargetEditingToolsStyle.h
new file mode 100644
--- /dev/null
+++ b/Engine/Plugins/Experimental/Animation/SkeletalMeshMorphTargetEditingTools/Source/SkeletalMeshMorphTargetEditingTools/Public/SkeletalMeshMorphTargetEditingToolsStyle.h
@@ -0,0 +1,22 @@
+// Copyright Epic Games, Inc. All Rights Reserved.
+
+#pragma once
+
+#include "Styling/SlateStyle.h"
+
+
+class FSkeletalMeshMorphTargetEditingToolsStyle
+	: public FSlateStyleSet
+{
+public:
+	static FSkeletalMeshMorphTargetEditingToolsStyle& Get();
+
+protected:
+	friend class FSkeletalMeshMorphTargetEditingToolsModule;
+
+	static void Register();
+	static void Unregister();
+
+private:
+	FSkeletalMeshMorphTargetEditingToolsStyle();
+};
\ No newline at end of file
diff --git a/Engine/Plugins/Experimental/Animation/SkeletalMeshMorphTargetEditingTools/Source/SkeletalMeshMorphTargetEditingTools/SkeletalMeshMorphTargetEditingTools.Build.cs b/Engine/Plugins/Experimental/Animation/SkeletalMeshMorphTargetEditingTools/Source/SkeletalMeshMorphTargetEditingTools/SkeletalMeshMorphTargetEditingTools.Build.cs
new file mode 100644
--- /dev/null
+++ b/Engine/Plugins/Experimental/Animation/SkeletalMeshMorphTargetEditingTools/Source/SkeletalMeshMorphTargetEditingTools/SkeletalMeshMorphTargetEditingTools.Build.cs
@@ -0,0 +1,67 @@
+﻿// Copyright Epic Games, Inc. All Rights Reserved.
+
+using UnrealBuildTool;
+
+public class SkeletalMeshMorphTargetEditingTools : ModuleRules
+{
+	public SkeletalMeshMorphTargetEditingTools(ReadOnlyTargetRules Target) : base(Target)
+	{
+		PCHUsage = ModuleRules.PCHUsageMode.UseExplicitOrSharedPCHs;
+		
+		PublicIncludePaths.AddRange(
+			new string[] {
+				// ... add public include paths required here ...
+			}
+			);
+				
+		
+		PrivateIncludePaths.AddRange(
+			new string[] {
+				// ... add other private include paths required here ...
+			}
+			);
+			
+		
+		PublicDependencyModuleNames.AddRange(
+			new string[]
+			{
+				"Core",
+				// ... add other public dependencies that you statically link with here ...
+				"ModelingToolsEditorMode",
+				"SkeletalMeshModelingTools"
+			}
+			);
+			
+		
+		PrivateDependencyModuleNames.AddRange(
+			new string[]
+			{
+				"CoreUObject",
+				"Engine",
+				"Slate",
+				"SlateCore", 
+				"ModelingComponents", 
+				"MeshConversion", 
+				"MeshModelingTools", 
+				"MeshModelingToolsEditorOnly", 
+				"InteractiveToolsFramework",
+				"SkeletalMeshUtilitiesCommon",
+				"GeometryCore",
+				"SkeletalMeshModifiers",
+				"MeshDescription",
+				"StaticMeshDescription",
+				// ... add private dependencies that you statically link with here ...	
+			}
+			);
+		
+		
+		DynamicallyLoadedModuleNames.AddRange(
+			new string[]
+			{
+				// ... add any modules that your module loads dynamically here ...
+			}
+			);
+		
+		ShortName = "SKMMorphTools";
+	}
+}
