[Skeletal Mesh Morph target editing tools] fix compilation error

diff --git a/Engine/Plugins/Experimental/Animation/SkeletalMeshMorphTargetEditingTools/Source/SkeletalMeshMorphTargetEditingTools/Public/MorphTargetModifier.h b/Engine/Plugins/Experimental/Animation/SkeletalMeshMorphTargetEditingTools/Source/SkeletalMeshMorphTargetEditingTools/Public/MorphTargetModifier.h
--- a/Engine/Plugins/Experimental/Animation/SkeletalMeshMorphTargetEditingTools/Source/SkeletalMeshMorphTargetEditingTools/Public/MorphTargetModifier.h
+++ b/Engine/Plugins/Experimental/Animation/SkeletalMeshMorphTargetEditingTools/Source/SkeletalMeshMorphTargetEditingTools/Public/MorphTargetModifier.h
@@ -5,6 +5,7 @@
 #include "Math/MathFwd.h"
 #include "Containers/Array.h"
 #include "Templates/UniquePtr.h"
+#include "MeshDescription.h"
 
 #include "MorphTargetModifier.generated.h"
 
@@ -12,7 +13,6 @@ class USkeletalMesh;
 class UMorphTarget;
 class FName;
 struct FMorphTargetDelta;
-struct FMeshDescription;
 
 UCLASS(BlueprintType, Blueprintable)
 class SKELETALMESHMORPHTARGETEDITINGTOOLS_API UMorphTargetModifier : public UObject
