[morph editing tool] fix loctext key conflict

diff --git a/Engine/Plugins/Experimental/Animation/SkeletalMeshMorphTargetEditingTools/Source/SkeletalMeshMorphTargetEditingTools/Private/MorphTargetEditorTool.cpp b/Engine/Plugins/Experimental/Animation/SkeletalMeshMorphTargetEditingTools/Source/SkeletalMeshMorphTargetEditingTools/Private/MorphTargetEditorTool.cpp
--- a/Engine/Plugins/Experimental/Animation/SkeletalMeshMorphTargetEditingTools/Source/SkeletalMeshMorphTargetEditingTools/Private/MorphTargetEditorTool.cpp
+++ b/Engine/Plugins/Experimental/Animation/SkeletalMeshMorphTargetEditingTools/Source/SkeletalMeshMorphTargetEditingTools/Private/MorphTargetEditorTool.cpp
@@ -71,7 +71,7 @@ void UMorphTargetEditorTool::HandleSkeletalMeshModified(const TArray<FName>& Bon
 
 void UMorphTargetEditorTool::CommitResult(UBaseDynamicMeshComponent* Component, bool bModifiedTopology)
 {
-	GetToolManager()->BeginUndoTransaction(LOCTEXT("SculptMeshToolTransactionName", "Sculpt Mesh"));
+	GetToolManager()->BeginUndoTransaction(LOCTEXT("MorphTargetEditorToolTransactionName", "Commit Morph Target Change"));
 	
 	switch (Properties->Operation)
 	{
@@ -90,15 +90,10 @@ void UMorphTargetManagerTool::Shutdown(EToolShutdownType ShutdownType)
 {
 	if (ShutdownType == EToolShutdownType::Accept)
 	{
-		GetToolManager()->BeginUndoTransaction(LOCTEXT("MorphTargetEditorTool", "Commit Skeleton Editing"));
+		GetToolManager()->BeginUndoTransaction(LOCTEXT("MorphTargetManagerToolTransactionName", "Commit Morph Target Changes"));
 		// Modifier->CommitMorphTargetsToSkeletalMesh();
 		GetToolManager()->EndUndoTransaction();
 
-		// to force to refresh the tree
-		if (NeedsNotification())
-		{
-			GetNotifier().Notify({}, ESkeletalMeshNotifyType::BonesAdded);
-		}
 	}
 	
 	Super::Shutdown(ShutdownType);
@@ -109,7 +104,7 @@ void UMorphTargetManagerTool::Shutdown(EToolShutdownType ShutdownType)
 
 void UMorphTargetManagerTool::RenameMorphTarget()
 {
-	GetToolManager()->BeginUndoTransaction(LOCTEXT("SculptMeshToolTransactionName", "Sculpt Mesh"));
+	GetToolManager()->BeginUndoTransaction(LOCTEXT("RenameMorphTargetTransactionName", "Rename Morph Target"));
 
 	// Modify SkeletalMesh
 	const FName& OldName = Properties->RenameOldName;
@@ -121,7 +116,7 @@ void UMorphTargetManagerTool::RenameMorphTarget()
 
 void UMorphTargetManagerTool::RemoveMorphTarget()
 {
-	GetToolManager()->BeginUndoTransaction(LOCTEXT("SculptMeshToolTransactionName", "Sculpt Mesh"));
+	GetToolManager()->BeginUndoTransaction(LOCTEXT("RemoveMorphTargetTransactionName", "Remove Morph Target"));
 
 	// Modify SkeletalMesh
 	const FName& Name = Properties->RemoveName;
@@ -132,7 +127,7 @@ void UMorphTargetManagerTool::RemoveMorphTarget()
 
 void UMorphTargetManagerTool::RemoveAllMorphTargets()
 {
-	GetToolManager()->BeginUndoTransaction(LOCTEXT("SculptMeshToolTransactionName", "Sculpt Mesh"));
+	GetToolManager()->BeginUndoTransaction(LOCTEXT("RemoveAllMorphTargetsTransactionName", "Remove All Morph Targets"));
 
 	// Modify SkeletalMesh
 	Modifier->RemoveAllMorphTargets();
@@ -232,11 +227,6 @@ void UMorphTargetEditorTool::Shutdown(EToolShutdownType ShutdownType)
 {
 	if (ShutdownType == EToolShutdownType::Accept)
 	{
-		// to force to refresh the tree
-		if (NeedsNotification())
-		{
-			GetNotifier().Notify({}, ESkeletalMeshNotifyType::BonesAdded);
-		}
 	}
 
 	Super::Shutdown(ShutdownType);
