[SkeletalMeshMorphTargetEditingTools] Fixed morph tool not working with posing.

+ fixed undo/redo for morph weight change

diff --git a/Engine/Plugins/Animation/SkeletalMeshModelingTools/Source/SkeletalMeshModelingTools/SkeletalMeshModelingTools.Build.cs b/Engine/Plugins/Animation/SkeletalMeshModelingTools/Source/SkeletalMeshModelingTools/SkeletalMeshModelingTools.Build.cs
--- a/Engine/Plugins/Animation/SkeletalMeshModelingTools/Source/SkeletalMeshModelingTools/SkeletalMeshModelingTools.Build.cs
+++ b/Engine/Plugins/Animation/SkeletalMeshModelingTools/Source/SkeletalMeshModelingTools/SkeletalMeshModelingTools.Build.cs
@@ -35,7 +35,14 @@ public class SkeletalMeshModelingTools : ModuleRules
 				"ApplicationCore",
 				"WidgetRegistration",
 				"SkeletalMeshModifiers",
-				"Persona"
+				"Persona",
+				"MeshDescription",
+				"StaticMeshDescription",
+				"SkeletalMeshDescription",
+				"GeometryCore",
+				"ModelingComponents",
+				"MeshConversion",
+				"InterchangeEngine"
 			}
 		);
 
diff --git a/Engine/Plugins/Experimental/Animation/SkeletalMeshMorphTargetEditingTools/Source/SkeletalMeshMorphTargetEditingTools/Private/IMorphTargetEditingToolInterface.cpp b/Engine/Plugins/Experimental/Animation/SkeletalMeshMorphTargetEditingTools/Source/SkeletalMeshMorphTargetEditingTools/Private/IMorphTargetEditingToolInterface.cpp
new file mode 100644
--- /dev/null
+++ b/Engine/Plugins/Experimental/Animation/SkeletalMeshMorphTargetEditingTools/Source/SkeletalMeshMorphTargetEditingTools/Private/IMorphTargetEditingToolInterface.cpp
@@ -0,0 +1,65 @@
+﻿// Copyright Epic Games, Inc. All Rights Reserved.
+
+#include "IMorphTargetEditingToolInterface.h"
+
+#include "ContextObjectStore.h"
+#include "IMorphTargetEditingToolInterface.h"
+#include "InteractiveToolManager.h"
+#include "IPersonaEditorModeManager.h"
+#include "MorphTargetEditingToolProperties.h"
+#include "PersonaModule.h"
+#include "SingleSelectionTool.h"
+#include "SkeletalMeshMorphTargetBackedTarget.h"
+#include "SkeletalMesh/SkeletalMeshEditionInterface.h"
+
+void IMorphTargetEditingToolInterface::SetupMorphEditingToolCommon()
+{
+	// Only support single selection tools at the moment
+	USingleSelectionTool* SingleTargetTool = CastChecked<USingleSelectionTool>(this);
+	ISkeletalMeshEditingInterface* EditingInterface = CastChecked<ISkeletalMeshEditingInterface>(this);
+	
+	if (USkeletalMeshEditorContextObjectBase* EditorContext = SingleTargetTool->GetToolManager()->GetContextObjectStore()->FindContext<USkeletalMeshEditorContextObjectBase>())
+	{
+		EditorContext->BindTo(EditingInterface);
+	}
+	
+	ISkeletalMeshMorphTargetBackedTarget* MorphTargetBackedTarget = Cast<ISkeletalMeshMorphTargetBackedTarget>(SingleTargetTool->GetTarget());
+
+	auto SetupFunction = [&](UMorphTargetEditingToolProperties* InProperties)
+	{
+		InProperties->MorphTargetNames = MorphTargetBackedTarget->GetEditableMorphTargetNames();
+		TArray<FName> SelectedMorphTargets = EditingInterface->GetSelectedMorphTargets();
+		
+		InProperties->EditMorphTargetName = SelectedMorphTargets.Num() > 0 ?
+			SelectedMorphTargets[0] : InProperties->MorphTargetNames.Num() > 0 ? InProperties->MorphTargetNames[0] : NAME_None;
+		
+		InProperties->NewMorphTargetName = MorphTargetBackedTarget->GetValidNameForNewMorphTarget(TEXT("NewMorphTarget"));
+
+		if (SelectedMorphTargets.Num() > 0)
+		{
+			InProperties->Operation = EMorphTargetEditorOperation::Edit;
+			MorphTargetBackedTarget->SetEditingMorphTargetName(InProperties->EditMorphTargetName);
+		}
+		else
+		{
+			InProperties->Operation = EMorphTargetEditorOperation::New;
+			MorphTargetBackedTarget->SetEditingMorphTargetName(InProperties->NewMorphTargetName);
+		}
+		
+	};
+	
+	SetupCommonProperties(SetupFunction);
+}
+
+void IMorphTargetEditingToolInterface::ShutdownMorphEditingToolCommon()
+{
+	ISkeletalMeshEditingInterface* EditingInterface = CastChecked<ISkeletalMeshEditingInterface>(this);
+
+	// Only support single selection tools at the moment
+	USingleSelectionTool* SingleTargetTool = CastChecked<USingleSelectionTool>(this);
+
+	if (USkeletalMeshEditorContextObjectBase* EditorContext = SingleTargetTool->GetToolManager()->GetContextObjectStore()->FindContext<USkeletalMeshEditorContextObjectBase>())
+	{
+		EditorContext->UnbindFrom(EditingInterface);
+	}	
+}
diff --git a/Engine/Plugins/Experimental/Animation/SkeletalMeshMorphTargetEditingTools/Source/SkeletalMeshMorphTargetEditingTools/Private/MorphTargetEditingToolProperties.cpp b/Engine/Plugins/Experimental/Animation/SkeletalMeshMorphTargetEditingTools/Source/SkeletalMeshMorphTargetEditingTools/Private/MorphTargetEditingToolProperties.cpp
new file mode 100644
--- /dev/null
+++ b/Engine/Plugins/Experimental/Animation/SkeletalMeshMorphTargetEditingTools/Source/SkeletalMeshMorphTargetEditingTools/Private/MorphTargetEditingToolProperties.cpp
@@ -0,0 +1,13 @@
+// Copyright Epic Games, Inc. All Rights Reserved.
+
+#include "MorphTargetEditingToolProperties.h"
+
+TArray<FName> UMorphTargetEditingToolProperties::GetMorphTargetNames()
+{
+	return MorphTargetNames;
+}
+
+FName UMorphTargetEditingToolProperties::GetEditingMorphTargetName()
+{
+	return Operation == EMorphTargetEditorOperation::New ? NewMorphTargetName : EditMorphTargetName;
+}
\ No newline at end of file
diff --git a/Engine/Plugins/Experimental/Animation/SkeletalMeshMorphTargetEditingTools/Source/SkeletalMeshMorphTargetEditingTools/Private/MorphTargetEditorTool.cpp b/Engine/Plugins/Experimental/Animation/SkeletalMeshMorphTargetEditingTools/Source/SkeletalMeshMorphTargetEditingTools/Private/MorphTargetEditorTool.cpp
deleted file mode 100644
diff --git a/Engine/Plugins/Experimental/Animation/SkeletalMeshMorphTargetEditingTools/Source/SkeletalMeshMorphTargetEditingTools/Private/MorphTargetModifer.cpp b/Engine/Plugins/Experimental/Animation/SkeletalMeshMorphTargetEditingTools/Source/SkeletalMeshMorphTargetEditingTools/Private/MorphTargetModifer.cpp
deleted file mode 100644
diff --git a/Engine/Plugins/Experimental/Animation/SkeletalMeshMorphTargetEditingTools/Source/SkeletalMeshMorphTargetEditingTools/Private/MorphTargetVertexSculptTool.cpp b/Engine/Plugins/Experimental/Animation/SkeletalMeshMorphTargetEditingTools/Source/SkeletalMeshMorphTargetEditingTools/Private/MorphTargetVertexSculptTool.cpp
new file mode 100644
--- /dev/null
+++ b/Engine/Plugins/Experimental/Animation/SkeletalMeshMorphTargetEditingTools/Source/SkeletalMeshMorphTargetEditingTools/Private/MorphTargetVertexSculptTool.cpp
@@ -0,0 +1,403 @@
+// Copyright Epic Games, Inc. All Rights Reserved.
+
+#include "MorphTargetVertexSculptTool.h"
+
+#include "ContextObjectStore.h"
+#include "ToolBuilderUtil.h"
+#include "InteractiveToolManager.h"
+#include "ModelingToolTargetUtil.h"
+#include "ToolTargetManager.h"
+#include "StaticMeshAttributes.h"
+#include "Engine/SkeletalMesh.h"
+#include "Components/SkeletalMeshComponent.h"
+#include "MeshDescriptionToDynamicMesh.h"
+#include "DynamicMeshToMeshDescription.h"
+#include "IPersonaEditorModeManager.h"
+#include "PersonaModule.h"
+#include "SkeletalMeshMorphTargetBackedTarget.h"
+#include "SkeletalMeshOperations.h"
+#include "BaseGizmos/BrushStampIndicator.h"
+#include "TargetInterfaces/DynamicMeshCommitter.h"
+#include "TargetInterfaces/DynamicMeshProvider.h"
+#include "TargetInterfaces/PrimitiveComponentBackedTarget.h"
+#include "AnimationRuntime.h"
+#include "SkeletalMeshAttributes.h"
+
+#define LOCTEXT_NAMESPACE "MorphTargetVertexSculptTool"
+
+
+UMeshSurfacePointTool* UMorphTargetVertexSculptToolBuilder::CreateNewTool(const FToolBuilderState& SceneState) const
+{
+	UMorphTargetVertexSculptTool* MorphTargetEditorTool = NewObject<UMorphTargetVertexSculptTool>(SceneState.ToolManager);
+	MorphTargetEditorTool->SetWorld(SceneState.World);
+	return MorphTargetEditorTool;
+}
+
+const FToolTargetTypeRequirements& UMorphTargetVertexSculptToolBuilder::GetTargetRequirements() const
+{
+	static FToolTargetTypeRequirements ToolRequirements({
+		UMaterialProvider::StaticClass(),
+		UDynamicMeshProvider::StaticClass(),
+		UDynamicMeshCommitter::StaticClass(),
+		USceneComponentBackedTarget::StaticClass(),
+		USkeletalMeshMorphTargetBackedTarget::StaticClass(),
+	});
+
+	return ToolRequirements;
+}
+
+
+void UMorphTargetVertexSculptTool::Setup()
+{
+	SetupMorphEditingToolCommon();
+	
+	// Setup Vertex Sculpt Tool
+	Super::Setup();
+
+	MorphTargetBackedTarget = Cast<ISkeletalMeshMorphTargetBackedTarget>(GetTarget());
+
+	SymmetryProperties->bSymmetryCanBeEnabled = false;
+
+	OnToolMeshChangedDelegate = DynamicMeshComponent->OnMeshRegionChanged.AddUObject(this, &UMorphTargetVertexSculptTool::OnToolMeshChanged);
+	
+	InitializeCache();
+}
+
+void UMorphTargetVertexSculptTool::Shutdown(EToolShutdownType ShutdownType)
+{
+	if (DynamicMeshComponent != nullptr)
+	{
+		DynamicMeshComponent->OnMeshChanged.Remove(OnDynamicMeshComponentChangedHandle);
+	}
+	
+	UpdateCacheIfNeeded();
+	
+	MorphTargetBackedTarget->SetEditingMorphTargetName(EditorToolProperties->GetEditingMorphTargetName());
+	MorphTargetBackedTarget->SetDataToCommit(&CachedDelta);
+	
+	Super::Shutdown(ShutdownType);
+
+	ShutdownMorphEditingToolCommon();	
+}
+
+void UMorphTargetVertexSculptTool::OnTick(float DeltaTime)
+{
+	if (UPersonaEditorModeManagerContext* PersonaModeManagerContext = GetToolManager()->GetContextObjectStore()->FindContext<UPersonaEditorModeManagerContext>())
+	{
+		if (!PersonaModeManagerContext->GetPersonaEditorModeManager()->IsModeActive(FPersonaEditModes::SkeletonSelection))
+		{
+			PersonaModeManagerContext->GetPersonaEditorModeManager()->ActivateMode(FPersonaEditModes::SkeletonSelection);	
+		}
+	}
+
+	Super::OnTick(DeltaTime);
+
+	if (InStroke())
+	{
+		bCached = false;
+	}
+	else 
+	{
+		PoseToolMesh();		
+	}
+}
+
+void UMorphTargetVertexSculptTool::OnPropertyModified(UObject* PropertySet, FProperty* Property)
+{
+	Super::OnPropertyModified(PropertySet, Property);
+
+	if (Property->GetFName() == GET_MEMBER_NAME_CHECKED(UMorphTargetEditingToolProperties, NewMorphTargetName))
+	{
+		EditorToolProperties->NewMorphTargetName = MorphTargetBackedTarget->GetValidNameForNewMorphTarget(EditorToolProperties->NewMorphTargetName);
+	}
+	else if (Property->GetFName() == GET_MEMBER_NAME_CHECKED(UMorphTargetEditingToolProperties, Operation))
+	{
+		if (EditorToolProperties->Operation == EMorphTargetEditorOperation::Edit)
+		{
+			if (EditorToolProperties->GetMorphTargetNames().Num() == 0)
+			{
+				EditorToolProperties->Operation = EMorphTargetEditorOperation::New;
+			}
+		}
+	}
+}
+
+void UMorphTargetVertexSculptTool::OnToolMeshChanged(UDynamicMeshComponent* Component, const FMeshRegionChangeBase* Change, bool bRevert)
+{
+	if (bRevert)
+	{
+		bCached = false;
+		UpdateCacheIfNeeded();
+	}
+}
+
+
+void UMorphTargetVertexSculptTool::SetupCommonProperties(const TFunction<void(UMorphTargetEditingToolProperties*)>& InSetupFunction)
+{
+	EditorToolProperties = NewObject<UMorphTargetEditingToolProperties>(this);
+	EditorToolProperties->SetFlags(RF_Transactional);
+
+	InSetupFunction(EditorToolProperties);
+	
+	AddToolPropertySource(EditorToolProperties);
+}
+
+void UMorphTargetVertexSculptTool::HandleSkeletalMeshModified(const TArray<FName>& Payload, const ESkeletalMeshNotifyType InNotifyType)
+{
+}
+
+void UMorphTargetVertexSculptTool::InitializeCache()
+{
+	UPrimitiveComponent* Component = Cast<IPrimitiveComponentBackedTarget>(GetTarget())->GetOwnerComponent();
+	USkeletalMeshComponent* SkeletalMeshComponent = Cast<USkeletalMeshComponent>(Component);
+	
+	FMeshDescription* MeshDescription = SkeletalMeshComponent->GetSkeletalMeshAsset()->GetMeshDescription(LodIndex0);
+	int32 NumVerts = MeshDescription->Vertices().Num();
+	
+	CachedDelta.Init(FVector3f::Zero(), NumVerts);
+
+	if (EditorToolProperties->Operation == EMorphTargetEditorOperation::Edit)
+	{
+		FSkeletalMeshAttributes Attributes(*MeshDescription);
+		FName CurrentMorphTarget = EditorToolProperties->GetEditingMorphTargetName();
+		check(Attributes.GetMorphTargetNames().Contains(CurrentMorphTarget));
+
+		TVertexAttributesRef<FVector3f> MorphTargetPosDeltaAttribute = Attributes.GetVertexMorphPositionDelta(CurrentMorphTarget);
+
+		for (FVertexID VertexID : MeshDescription->Vertices().GetElementIDs())
+		{
+			CachedDelta[VertexID] = MorphTargetPosDeltaAttribute[VertexID];
+		}	
+	}
+
+	bCached = true;
+	
+}
+
+void UMorphTargetVertexSculptTool::UpdateCacheIfNeeded()
+{
+	if (bCached)
+	{
+		return;
+	}
+
+	// Don't update the cache if the morph cannot be extracted from the mesh
+	if (FMath::IsNearlyZero(EditorToolProperties->MorphTargetWeight))
+	{
+		bCached = true;
+		return;
+	}
+	
+	TArray<FVector3f> Results;
+	
+	UPrimitiveComponent* Component = Cast<IPrimitiveComponentBackedTarget>(GetTarget())->GetOwnerComponent();
+	USkeletalMeshComponent* SkeletalMeshComponent = Cast<USkeletalMeshComponent>(Component);
+	const FReferenceSkeleton& RefSkeleton = SkeletalMeshComponent->GetSkeletalMeshAsset()->GetRefSkeleton();
+	
+	TArray<FTransform> RefComponentSpaceTransforms;
+	for (int32 BoneIndex = 0; BoneIndex < RefSkeleton.GetNum(); ++BoneIndex)
+	{
+		// Need local to component transform
+		RefComponentSpaceTransforms.Add(FAnimationRuntime::GetComponentSpaceTransformRefPose(RefSkeleton, BoneIndex));
+	}
+
+	FName CurrentMorphTarget = EditorToolProperties->GetEditingMorphTargetName();
+
+	TMap<FName, float> MorphTargetWeights = PreviousMorphWeights;
+	MorphTargetWeights.Remove(CurrentMorphTarget);
+	
+	FMeshDescription SculptMeshDescription;
+	FStaticMeshAttributes Attributes(SculptMeshDescription);
+	Attributes.Register();
+	
+	FMeshDescription NonDeformedMesh;
+
+	FDynamicMeshToMeshDescription Converter;
+	Converter.Convert(GetSculptMesh(), SculptMeshDescription, false);
+	
+	const FMeshDescription* RefMeshDescription = SkeletalMeshComponent->GetSkeletalMeshAsset()->GetMeshDescription(LodIndex0);
+
+	const TArray<FTransform>& ComponentSpaceTransforms = PreviousPoseComponentSpace;
+	
+	if (FSkeletalMeshOperations::GetMorphTargetDeltasFromPosedMesh(SculptMeshDescription, *RefMeshDescription, RefComponentSpaceTransforms, NonDeformedMesh, ComponentSpaceTransforms, NAME_None, MorphTargetWeights))
+	{
+		for (FVertexID VertexID : NonDeformedMesh.Vertices().GetElementIDs())
+		{
+			FVector3f Delta = NonDeformedMesh.GetVertexPosition(VertexID) - RefMeshDescription->GetVertexPosition(VertexID);
+			Delta /= EditorToolProperties->MorphTargetWeight;
+			
+			Results.Add(Delta);
+		}
+	}
+
+	CachedDelta = MoveTemp(Results);
+
+	bCached = true;
+
+	GetToolManager()->DisplayMessage(LOCTEXT("UpdateCacheNotification", "Cache Updated"), EToolMessageLevel::UserNotification);
+}
+
+void UMorphTargetVertexSculptTool::PoseToolMesh()
+{
+	if (!AllowToolMeshUpdates())
+	{
+		return;
+	}
+	
+	auto IsMorphWeightChanged = [&PreviousMorphWeights = this->PreviousMorphWeights](const FName& InName, float InMorphWeight) -> bool
+	{
+		if (float* PreviousWeight = PreviousMorphWeights.Find(InName))
+		{
+			if (!FMath::IsNearlyEqual(*PreviousWeight, InMorphWeight))
+			{
+				return true;
+			}
+		}
+		else
+		{
+			return true;
+		}
+		return false;
+	};
+	
+	
+	UPrimitiveComponent* Component = Cast<IPrimitiveComponentBackedTarget>(GetTarget())->GetOwnerComponent();
+	USkeletalMeshComponent* SkeletalMeshComponent = Cast<USkeletalMeshComponent>(Component);
+	
+	FName CurrentMorphTarget = EditorToolProperties->GetEditingMorphTargetName();
+		
+	TMap<FName, float> MorphTargetWeights;
+
+	for (const TPair<const UMorphTarget*, int32>& MorphTarget: SkeletalMeshComponent->ActiveMorphTargets)
+	{
+		const FName MorphName = MorphTarget.Key->GetFName();
+		const float MorphWeight = SkeletalMeshComponent->MorphTargetWeights[MorphTarget.Value];
+
+		MorphTargetWeights.Add(MorphName, MorphWeight);
+	}
+	
+	MorphTargetWeights.FindOrAdd(CurrentMorphTarget) = EditorToolProperties->MorphTargetWeight;
+
+	const TArray<FTransform>& ComponentSpaceTransforms = SkeletalMeshComponent->GetComponentSpaceTransforms();
+	
+	bool bPoseChanged = false;
+	
+	// First tick, force an update
+	if (!bPoseOpInitialized)
+	{
+		bPoseOpInitialized = true;
+		
+		PreviousMorphWeights = MorphTargetWeights;
+		PreviousPoseComponentSpace = ComponentSpaceTransforms;
+		bPoseChanged = true;
+	}
+	else
+	{
+		// Check for pose changes and update accordingly
+		
+		if (!bPoseChanged)
+		{
+			if (MorphTargetWeights.Num() != PreviousMorphWeights.Num())
+			{
+				bPoseChanged = true;
+			}
+		}
+		
+		if (!bPoseChanged)
+		{
+			for (const TPair<FName, float>& MorphWeight: MorphTargetWeights)
+			{
+				if (IsMorphWeightChanged(MorphWeight.Key, MorphWeight.Value))
+				{
+					bPoseChanged = true;
+					break;
+				}
+			}
+		}
+		
+		if (!bPoseChanged)
+		{
+			for (int32 BoneIndex=0; BoneIndex< ComponentSpaceTransforms.Num(); ++BoneIndex)
+			{
+				const FTransform& CurrentBoneTransform = ComponentSpaceTransforms[BoneIndex];
+				const FTransform& PrevBoneTransform = PreviousPoseComponentSpace[BoneIndex];
+				if (!CurrentBoneTransform.Equals(PrevBoneTransform))
+				{
+					bPoseChanged = true;
+					break;
+				}
+			}
+		}
+	}
+
+	if (!bPoseChanged)
+	{
+		if (bPoseChangedLastTick)
+		{
+			FMeshReplacementChange Change(MeshBeforePosing, MakeShared<FDynamicMesh3>(*GetSculptMesh()));
+			
+			OnDynamicMeshComponentChanged(DynamicMeshComponent, &Change, false);
+		}
+
+		bPoseChangedLastTick = bPoseChanged;
+		return;
+	}
+
+	if (!bPoseChangedLastTick)
+	{
+		MeshBeforePosing = MakeShared<FDynamicMesh3>(*GetSculptMesh());
+	}
+	bPoseChangedLastTick = bPoseChanged;
+	
+	
+	UpdateCacheIfNeeded();
+	
+	PreviousPoseComponentSpace = ComponentSpaceTransforms;
+	PreviousMorphWeights = MorphTargetWeights;
+
+	// have to wait for any outstanding stamp/undo update to finish...
+	WaitForPendingStampUpdate();
+	WaitForPendingUndoRedo();
+
+	TArray<FTransform> RefComponentSpaceTransforms;
+	for (int32 BoneIndex = 0; BoneIndex < ComponentSpaceTransforms.Num(); ++BoneIndex)
+	{
+		// Need local to component transform
+		RefComponentSpaceTransforms.Add(FAnimationRuntime::GetComponentSpaceTransformRefPose(SkeletalMeshComponent->GetSkeletalMeshAsset()->GetRefSkeleton(), BoneIndex));
+	}
+			
+	FMeshDescription PosedMeshDescription;
+
+	{
+		FMeshDescription MeshDescriptionWithAllMorphs = *SkeletalMeshComponent->GetSkeletalMeshAsset()->GetMeshDescription(LodIndex0);
+		FSkeletalMeshAttributes Attributes(MeshDescriptionWithAllMorphs);
+
+		if (!Attributes.GetMorphTargetNames().Contains(CurrentMorphTarget))
+		{
+			Attributes.RegisterMorphTargetAttribute(CurrentMorphTarget, false);
+		}
+
+		TVertexAttributesRef<FVector3f> MorphTargetPosDeltaAttribute = Attributes.GetVertexMorphPositionDelta(CurrentMorphTarget);
+
+		for (FVertexID VertexID : MeshDescriptionWithAllMorphs.Vertices().GetElementIDs())
+		{
+			MorphTargetPosDeltaAttribute[VertexID] = CachedDelta[VertexID];
+		}
+
+		FSkeletalMeshOperations::GetPosedMesh(MeshDescriptionWithAllMorphs, PosedMeshDescription, ComponentSpaceTransforms, NAME_None, MorphTargetWeights);
+	}
+
+	FDynamicMesh3& Mesh = *GetSculptMesh();
+	
+	TArrayView<const FVector3f> VertexPositionsAttribute = PosedMeshDescription.GetVertexPositions().GetRawArray();
+	for (int32 Index = 0; Index < VertexPositionsAttribute.Num(); ++Index)
+	{
+		Mesh.SetVertex(Index, FVector3d(VertexPositionsAttribute[Index]));
+	}	
+	
+	DynamicMeshComponent->FastNotifyPositionsUpdated();
+}	
+
+
+#undef LOCTEXT_NAMESPACE
+
diff --git a/Engine/Plugins/Experimental/Animation/SkeletalMeshMorphTargetEditingTools/Source/SkeletalMeshMorphTargetEditingTools/Private/SkeletalMeshMorphTargetEditingToolsCommands.cpp b/Engine/Plugins/Experimental/Animation/SkeletalMeshMorphTargetEditingTools/Source/SkeletalMeshMorphTargetEditingTools/Private/SkeletalMeshMorphTargetEditingToolsCommands.cpp
--- a/Engine/Plugins/Experimental/Animation/SkeletalMeshMorphTargetEditingTools/Source/SkeletalMeshMorphTargetEditingTools/Private/SkeletalMeshMorphTargetEditingToolsCommands.cpp
+++ b/Engine/Plugins/Experimental/Animation/SkeletalMeshMorphTargetEditingTools/Source/SkeletalMeshMorphTargetEditingTools/Private/SkeletalMeshMorphTargetEditingToolsCommands.cpp
@@ -21,8 +21,7 @@ void FSkeletalMeshMorphTargetEditingToolsCommands::RegisterCommands()
 {
 	UI_COMMAND(BeginMorphTargetTool, "Morph", "Edit Morph Targets", EUserInterfaceActionType::ToggleButton, FInputChord());
 	
-	UI_COMMAND(BeginMorphTargetEditorTool, "Add/Edit Morph Target", "MorphTarget Sculpt Tool", EUserInterfaceActionType::ToggleButton, FInputChord());
-	UI_COMMAND(BeginMorphTargetManagerTool, "Manange Morph Targets ", "MorphTarget Manager Tool", EUserInterfaceActionType::ToggleButton, FInputChord());
+	UI_COMMAND(BeginMorphTargetEditorTool, "Sculpt Morph Target", "MorphTarget Sculpt Tool", EUserInterfaceActionType::ToggleButton, FInputChord());
 }
 
 #undef LOCTEXT_NAMESPACE
diff --git a/Engine/Plugins/Experimental/Animation/SkeletalMeshMorphTargetEditingTools/Source/SkeletalMeshMorphTargetEditingTools/Private/SkeletalMeshMorphTargetEditingToolsModule.cpp b/Engine/Plugins/Experimental/Animation/SkeletalMeshMorphTargetEditingTools/Source/SkeletalMeshMorphTargetEditingTools/Private/SkeletalMeshMorphTargetEditingToolsModule.cpp
--- a/Engine/Plugins/Experimental/Animation/SkeletalMeshMorphTargetEditingTools/Source/SkeletalMeshMorphTargetEditingTools/Private/SkeletalMeshMorphTargetEditingToolsModule.cpp
+++ b/Engine/Plugins/Experimental/Animation/SkeletalMeshMorphTargetEditingTools/Source/SkeletalMeshMorphTargetEditingTools/Private/SkeletalMeshMorphTargetEditingToolsModule.cpp
@@ -2,9 +2,10 @@
 
 #include "SkeletalMeshMorphTargetEditingToolsModule.h"
 
-#include "MorphTargetEditorTool.h"
+#include "MorphTargetVertexSculptTool.h"
 #include "SkeletalMeshMorphTargetEditingToolsCommands.h"
 #include "SkeletalMeshMorphTargetEditingToolsStyle.h"
+#include "SkeletalMeshMorphTargetToolTarget.h"
 #include "Features/IModularFeatures.h"
 #include "Framework/Commands/Commands.h"
 
@@ -33,19 +34,11 @@ void FSkeletalMeshMorphTargetEditingToolsModule::GetExtensionTools(const FExtens
 	const FSkeletalMeshMorphTargetEditingToolsCommands& Commands = FSkeletalMeshMorphTargetEditingToolsCommands::Get();
 	{
 		FExtensionToolDescription ToolDesc;
-		ToolDesc.ToolName = LOCTEXT("SkeletalMeshMorphTargetEditorTool", "Add/Edit Morph Target");
+		ToolDesc.ToolName = LOCTEXT("SkeletalMeshMorphTargetEditorTool", "Sculpt Morph Target");
 		ToolDesc.ToolCommand = Commands.BeginMorphTargetEditorTool;
-		ToolDesc.ToolBuilder = NewObject<UMorphTargetEditorToolBuilder>();
+		ToolDesc.ToolBuilder = NewObject<UMorphTargetVertexSculptToolBuilder>();
 		OutTools.Add(ToolDesc);
 	}
-
-	{
-		FExtensionToolDescription ToolDesc;
-		ToolDesc.ToolName = LOCTEXT("SkeletalMeshMorphTargetManagerTool", "Manange Morph Targets");
-		ToolDesc.ToolCommand = Commands.BeginMorphTargetManagerTool;
-		ToolDesc.ToolBuilder = NewObject<UMorphTargetManagerToolBuilder>();
-		OutTools.Add(ToolDesc);
-	}	
 }
 
 bool FSkeletalMeshMorphTargetEditingToolsModule::GetExtensionExtendedInfo(FModelingModeExtensionExtendedInfo& InfoOut)
@@ -55,6 +48,14 @@ bool FSkeletalMeshMorphTargetEditingToolsModule::GetExtensionExtendedInfo(FModel
 	return true;
 }
 
+bool FSkeletalMeshMorphTargetEditingToolsModule::GetExtensionToolTargets(TArray<TSubclassOf<UToolTargetFactory>>& ToolTargetFactoriesOut)
+{
+	ToolTargetFactoriesOut.Add(USkeletalMeshMorphTargetToolTargetFactory::StaticClass());
+	
+	return true;
+}
+
+
 #undef LOCTEXT_NAMESPACE
 	
 IMPLEMENT_MODULE(FSkeletalMeshMorphTargetEditingToolsModule, SkeletalMeshMorphTargetEditingTools)
\ No newline at end of file
diff --git a/Engine/Plugins/Experimental/Animation/SkeletalMeshMorphTargetEditingTools/Source/SkeletalMeshMorphTargetEditingTools/Private/SkeletalMeshMorphTargetToolTarget.cpp b/Engine/Plugins/Experimental/Animation/SkeletalMeshMorphTargetEditingTools/Source/SkeletalMeshMorphTargetEditingTools/Private/SkeletalMeshMorphTargetToolTarget.cpp
new file mode 100644
--- /dev/null
+++ b/Engine/Plugins/Experimental/Animation/SkeletalMeshMorphTargetEditingTools/Source/SkeletalMeshMorphTargetEditingTools/Private/SkeletalMeshMorphTargetToolTarget.cpp
@@ -0,0 +1,157 @@
+﻿// Copyright Epic Games, Inc. All Rights Reserved.
+
+#include "SkeletalMeshMorphTargetToolTarget.h"
+
+#include "DynamicMeshToMeshDescription.h"
+#include "MeshDescriptionToDynamicMesh.h"
+#include "SkeletalMeshAttributes.h"
+#include "SkeletalMeshOperations.h"
+#include "SkeletalMeshTypes.h"
+#include "StaticMeshAttributes.h"
+#include "Components/SkeletalMeshComponent.h"
+#include "ConversionUtils/DynamicMeshViaMeshDescriptionUtil.h"
+#include "UObject/Package.h"
+#include "ConversionUtils/SkinnedMeshToDynamicMesh.h"
+#include "Rendering/SkeletalMeshModel.h"
+#include "AnimationRuntime.h"
+
+UE::Geometry::FDynamicMesh3 USkeletalMeshMorphTargetToolTarget::GetDynamicMesh()
+{
+	FGetMeshParameters Params;
+	return GetDynamicMesh(Params);
+}
+
+UE::Geometry::FDynamicMesh3 USkeletalMeshMorphTargetToolTarget::GetDynamicMesh(const FGetMeshParameters& InGetMeshParams)
+{
+	FDynamicMesh3 DynamicMesh;
+	FMeshDescriptionToDynamicMesh Converter;
+	Converter.Convert(GetEditingMeshDescription(), DynamicMesh);
+	
+	return DynamicMesh;
+}
+
+void USkeletalMeshMorphTargetToolTarget::CommitDynamicMesh(const UE::Geometry::FDynamicMesh3& Mesh, const FDynamicMeshCommitInfo& CommitInfo)
+{
+	if (!DataPtr->IsEmpty())
+	{
+		FMeshDescription* TargetMeshDescription = GetEditingMeshDescription();
+		FSkeletalMeshAttributes TargetAttributes(*TargetMeshDescription);
+
+		if (!TargetAttributes.GetMorphTargetNames().Contains(EditingMorphTargetName))
+		{
+			TargetAttributes.RegisterMorphTargetAttribute(EditingMorphTargetName, false);
+		}
+		
+		TVertexAttributesRef<FVector3f> TargetMorphTargetPosDeltaAttribute = TargetAttributes.GetVertexMorphPositionDelta(EditingMorphTargetName);
+
+		for (FVertexID VertexID : TargetMeshDescription->Vertices().GetElementIDs())
+		{
+			TargetMorphTargetPosDeltaAttribute[VertexID] = (*DataPtr)[VertexID.GetValue()];
+		}		
+	
+		CommitEditedMeshDescription();
+
+		USkeletalMesh* SkeletalMesh = GetSkeletalMesh();
+		UAnimCurveMetaData* AnimCurveMetaData = SkeletalMesh->GetAssetUserData<UAnimCurveMetaData>();
+		if (AnimCurveMetaData == nullptr)
+		{
+			AnimCurveMetaData = NewObject<UAnimCurveMetaData>(SkeletalMesh, NAME_None, RF_Transactional);
+			SkeletalMesh->AddAssetUserData(AnimCurveMetaData);
+		}
+
+		AnimCurveMetaData->AddCurveMetaData(EditingMorphTargetName);
+
+		// Ensure we have a morph flag set
+		if (FCurveMetaData* CurveMetaData = AnimCurveMetaData->GetCurveMetaData(EditingMorphTargetName))
+		{
+			CurveMetaData->Type.bMorphtarget = true;
+		}
+	}
+}
+
+void USkeletalMeshMorphTargetToolTarget::SetDataToCommit(TArray<FVector3f>* InDataPtr)
+{
+	DataPtr = InDataPtr;
+}
+
+void USkeletalMeshMorphTargetToolTarget::SetEditingMorphTargetName(const FName& InName)
+{
+	EditingMorphTargetName = InName;
+}
+
+TArray<FName> USkeletalMeshMorphTargetToolTarget::GetEditableMorphTargetNames()
+{
+	FSkeletalMeshAttributes TargetAttributes(*GetEditingMeshDescription());
+
+	return TargetAttributes.GetMorphTargetNames();
+}
+
+FName USkeletalMeshMorphTargetToolTarget::GetValidNameForNewMorphTarget(const FName& InName)
+{
+	FSkeletalMeshAttributes TargetAttributes(*GetEditingMeshDescription());
+
+	const TArray<FName>& MorphTargetNames = TargetAttributes.GetMorphTargetNames();
+	FName NewMorphTargetName = InName;
+	while (MorphTargetNames.Contains(NewMorphTargetName))
+	{
+		NewMorphTargetName.SetNumber(NewMorphTargetName.GetNumber() + 1);
+	}
+
+	return NewMorphTargetName;
+}
+
+FMeshDescription* USkeletalMeshMorphTargetToolTarget::GetEditingMeshDescription()
+{
+	USkeletalMesh* SkeletalMesh = GetSkeletalMesh();
+	FMeshDescription* TargetMeshDescription = SkeletalMesh->GetMeshDescription(LodIndex0);
+
+	return TargetMeshDescription;
+}
+
+void USkeletalMeshMorphTargetToolTarget::CommitEditedMeshDescription()
+{
+	USkeletalMesh* SkeletalMesh = GetSkeletalMesh();
+	{
+		FScopedSkeletalMeshPostEditChange ScopePostEditChange(SkeletalMesh);
+		SkeletalMesh->PreEditChange(nullptr);
+		SkeletalMesh->ModifyMeshDescription(LodIndex0);
+		SkeletalMesh->CommitMeshDescription(LodIndex0);
+	}
+}
+
+bool USkeletalMeshMorphTargetToolTargetFactory::CanBuildTarget(UObject* SourceObject,
+                                                               const FToolTargetTypeRequirements& Requirements) const
+{
+	// We are using an exact cast here to prevent subclasses, which might not meet all
+	// requirements for functionality such as the deprecated DestructibleMeshComponent, from 
+	// being caught up as valid targets.
+	// If you want to make the tool target work with some subclass of USkeletalMeshComponent,
+	// just add another factory that allows that class specifically(but make sure that
+	// GetMeshDescription and such work properly)
+	
+	bool bValid = Cast<USkeletalMeshComponent>(SourceObject) && Cast<USkeletalMeshComponent>(SourceObject)->GetSkeletalMeshAsset() &&
+		ExactCast<USkeletalMesh>(Cast<USkeletalMeshComponent>(SourceObject)->GetSkeletalMeshAsset()) &&
+		!ExactCast<USkeletalMesh>(Cast<USkeletalMeshComponent>(SourceObject)->GetSkeletalMeshAsset())->GetOutermost()->bIsCookedForEditor;
+	if (!bValid)
+	{
+		return false;
+	}
+
+	if (!USkeletalMeshComponentToolTargetFactory::CanWriteToSource(SourceObject))
+	{
+		return false;
+	}
+
+	return Requirements.AreSatisfiedBy(USkeletalMeshMorphTargetToolTarget::StaticClass());
+}
+
+UToolTarget* USkeletalMeshMorphTargetToolTargetFactory::BuildTarget(UObject* SourceObject, const FToolTargetTypeRequirements& TargetTypeInfo)
+{
+	USkeletalMeshMorphTargetToolTarget* Target = NewObject<USkeletalMeshMorphTargetToolTarget>();
+
+	Target->InitializeComponent(Cast<USkeletalMeshComponent>(SourceObject));
+	
+	return Target;
+	
+}
+
diff --git a/Engine/Plugins/Experimental/Animation/SkeletalMeshMorphTargetEditingTools/Source/SkeletalMeshMorphTargetEditingTools/Public/IMorphTargetEditingToolInterface.h b/Engine/Plugins/Experimental/Animation/SkeletalMeshMorphTargetEditingTools/Source/SkeletalMeshMorphTargetEditingTools/Public/IMorphTargetEditingToolInterface.h
new file mode 100644
--- /dev/null
+++ b/Engine/Plugins/Experimental/Animation/SkeletalMeshMorphTargetEditingTools/Source/SkeletalMeshMorphTargetEditingTools/Public/IMorphTargetEditingToolInterface.h
@@ -0,0 +1,29 @@
+﻿// Copyright Epic Games, Inc. All Rights Reserved.
+
+#pragma once
+
+#include "CoreMinimal.h"
+#include "UObject/Interface.h"
+
+#include "IMorphTargetEditingToolInterface.generated.h"
+
+class FMorphTargetEditChange;
+class UMorphTargetEditingToolProperties;
+
+UINTERFACE(MinimalAPI)
+class UMorphTargetEditingToolInterface : public UInterface
+{
+	GENERATED_BODY()
+};
+
+class IMorphTargetEditingToolInterface : public IInterface
+{
+	GENERATED_BODY()
+public:
+	virtual void SetupMorphEditingToolCommon();
+	virtual void ShutdownMorphEditingToolCommon();
+	virtual void SetupCommonProperties(const TFunction<void(UMorphTargetEditingToolProperties*)>& InSetupFunction) = 0;
+	// virtual void Apply(const FMorphTargetEditChange& Change) = 0;
+	// virtual void Revert(const FMorphTargetEditChange& Change) = 0;
+};
+
diff --git a/Engine/Plugins/Experimental/Animation/SkeletalMeshMorphTargetEditingTools/Source/SkeletalMeshMorphTargetEditingTools/Public/MorphTargetEditingToolProperties.h b/Engine/Plugins/Experimental/Animation/SkeletalMeshMorphTargetEditingTools/Source/SkeletalMeshMorphTargetEditingTools/Public/MorphTargetEditingToolProperties.h
new file mode 100644
--- /dev/null
+++ b/Engine/Plugins/Experimental/Animation/SkeletalMeshMorphTargetEditingTools/Source/SkeletalMeshMorphTargetEditingTools/Public/MorphTargetEditingToolProperties.h
@@ -0,0 +1,44 @@
+// Copyright Epic Games, Inc. All Rights Reserved.
+#pragma once
+#include "InteractiveTool.h"
+
+#include "MorphTargetEditingToolProperties.generated.h"
+
+UENUM()
+enum class EMorphTargetEditorOperation
+{
+	// Add new MorphTarget, select a morph target before activating the tool to edit
+	New,
+	// Edit existing MorphTarget, deselect morph targets before activating the tool to create a new morph target
+	Edit, 
+};
+
+UCLASS()
+class UMorphTargetEditingToolProperties : public UInteractiveToolPropertySet
+{
+	GENERATED_BODY()
+public:
+
+	// Select a morph target before activating the tool to edit, if nothing is selected, a new morph target is created
+	UPROPERTY(VisibleAnywhere, Category = "MorphTarget")
+	EMorphTargetEditorOperation Operation;
+
+	UPROPERTY(EditAnywhere, Category = "MorphTarget", meta = (DisplayName = "Name", EditCondition = "Operation == EMorphTargetEditorOperation::New", EditConditionHides))
+	FName NewMorphTargetName;
+	
+	UPROPERTY(VisibleAnywhere, Category = "MorphTarget", meta = (GetOptions = GetMorphTargetNames, NoResetToDefault, DisplayName = "Name", EditCondition = "Operation == EMorphTargetEditorOperation::Edit", EditConditionHides))
+	FName EditMorphTargetName;
+
+	UPROPERTY(EditAnywhere, Category = "MorphTarget", meta = (DisplayName = "Weight", UIMin = "0.0", UIMax = "1.0"))
+	float MorphTargetWeight = 1.f;
+
+	UFUNCTION()
+	TArray<FName> GetMorphTargetNames();
+
+	UPROPERTY()
+	TArray<FName> MorphTargetNames;
+
+	FName GetEditingMorphTargetName();
+
+	
+};
\ No newline at end of file
diff --git a/Engine/Plugins/Experimental/Animation/SkeletalMeshMorphTargetEditingTools/Source/SkeletalMeshMorphTargetEditingTools/Public/MorphTargetEditorTool.h b/Engine/Plugins/Experimental/Animation/SkeletalMeshMorphTargetEditingTools/Source/SkeletalMeshMorphTargetEditingTools/Public/MorphTargetEditorTool.h
deleted file mode 100644
diff --git a/Engine/Plugins/Experimental/Animation/SkeletalMeshMorphTargetEditingTools/Source/SkeletalMeshMorphTargetEditingTools/Public/MorphTargetModifier.h b/Engine/Plugins/Experimental/Animation/SkeletalMeshMorphTargetEditingTools/Source/SkeletalMeshMorphTargetEditingTools/Public/MorphTargetModifier.h
deleted file mode 100644
diff --git a/Engine/Plugins/Experimental/Animation/SkeletalMeshMorphTargetEditingTools/Source/SkeletalMeshMorphTargetEditingTools/Public/MorphTargetVertexSculptTool.h b/Engine/Plugins/Experimental/Animation/SkeletalMeshMorphTargetEditingTools/Source/SkeletalMeshMorphTargetEditingTools/Public/MorphTargetVertexSculptTool.h
new file mode 100644
--- /dev/null
+++ b/Engine/Plugins/Experimental/Animation/SkeletalMeshMorphTargetEditingTools/Source/SkeletalMeshMorphTargetEditingTools/Public/MorphTargetVertexSculptTool.h
@@ -0,0 +1,82 @@
+// Copyright Epic Games, Inc. All Rights Reserved.
+
+#pragma once
+
+#include "IMorphTargetEditingToolInterface.h"
+#include "MeshVertexSculptTool.h"
+#include "SkeletalMesh/SkeletalMeshEditionInterface.h"
+#include "BaseTools/MultiSelectionMeshEditingTool.h"
+#include "MorphTargetEditingToolProperties.h"
+#include "SkeletalMeshMorphTargetBackedTarget.h"
+
+#include "MorphTargetVertexSculptTool.generated.h"
+
+/**
+ * MorphTarget Vertex Sculpt Tool Builder
+ */
+UCLASS()
+class UMorphTargetVertexSculptToolBuilder : public UMeshVertexSculptToolBuilder
+{
+	GENERATED_BODY()
+
+public:
+	virtual UMeshSurfacePointTool* CreateNewTool(const FToolBuilderState& SceneState) const override;
+	virtual const FToolTargetTypeRequirements& GetTargetRequirements() const override;
+};
+
+
+/**
+* MorphTarget Editor Tool
+*/
+UCLASS()
+class UMorphTargetVertexSculptTool : 
+	public UMeshVertexSculptTool,
+	public ISkeletalMeshEditingInterface,
+	public IMorphTargetEditingToolInterface
+{
+	GENERATED_BODY()
+
+public:	
+	// UMeshVertexSculptTool overrides
+	virtual void Setup() override;
+	virtual void Shutdown(EToolShutdownType ShutdownType) override;
+	virtual void OnTick(float DeltaTime) override;
+	virtual void OnPropertyModified(UObject* PropertySet, FProperty* Property) override;
+
+	void OnToolMeshChanged(UDynamicMeshComponent* Component, const FMeshRegionChangeBase* Change, bool bRevert);
+
+	// IMorphTargetEditingTool overrides
+	virtual void SetupCommonProperties(const TFunction<void(UMorphTargetEditingToolProperties*)>& InSetupFunction) override;
+	
+	// ISkeletalMeshEditingInterface
+	virtual void HandleSkeletalMeshModified(const TArray<FName>& Payload, const ESkeletalMeshNotifyType InNotifyType) override;
+
+	void InitializeCache();
+	void UpdateCacheIfNeeded();
+
+	void PoseToolMesh();
+
+	UPROPERTY()
+	TObjectPtr<class UMorphTargetEditingToolProperties> EditorToolProperties;
+
+	TWeakInterfacePtr<ISkeletalMeshMorphTargetBackedTarget> MorphTargetBackedTarget; 
+	
+	bool bCached = false;
+	TArray<FVector3f> CachedDelta;
+
+	bool bPoseOpInitialized = false;
+	TArray<FTransform> PreviousPoseComponentSpace;
+	TMap<FName, float> PreviousMorphWeights;
+
+	TSharedPtr<FDynamicMesh3> MeshBeforePosing;
+	bool bPoseChangedLastTick = false;
+
+	FDelegateHandle OnToolMeshChangedDelegate;
+
+	static constexpr int32 LodIndex0 = 0;	
+};
+
+
+
+
+
diff --git a/Engine/Plugins/Experimental/Animation/SkeletalMeshMorphTargetEditingTools/Source/SkeletalMeshMorphTargetEditingTools/Public/SkeletalMeshMorphTargetBackedTarget.h b/Engine/Plugins/Experimental/Animation/SkeletalMeshMorphTargetEditingTools/Source/SkeletalMeshMorphTargetEditingTools/Public/SkeletalMeshMorphTargetBackedTarget.h
new file mode 100644
--- /dev/null
+++ b/Engine/Plugins/Experimental/Animation/SkeletalMeshMorphTargetEditingTools/Source/SkeletalMeshMorphTargetEditingTools/Public/SkeletalMeshMorphTargetBackedTarget.h
@@ -0,0 +1,25 @@
+﻿// Copyright Epic Games, Inc. All Rights Reserved.
+
+#pragma once
+
+#include "CoreMinimal.h"
+#include "UObject/Interface.h"
+
+#include "SkeletalMeshMorphTargetBackedTarget.generated.h"
+
+UINTERFACE(MinimalAPI)
+class USkeletalMeshMorphTargetBackedTarget : public UInterface
+{
+	GENERATED_BODY()
+};
+
+class ISkeletalMeshMorphTargetBackedTarget : public IInterface
+{
+	GENERATED_BODY()
+public:
+	virtual void SetDataToCommit(TArray<FVector3f>* InDataPtr) = 0;
+	virtual void SetEditingMorphTargetName(const FName& InName) = 0;
+	virtual TArray<FName> GetEditableMorphTargetNames() = 0;
+	virtual FName GetValidNameForNewMorphTarget(const FName& InName) = 0;
+};
+
diff --git a/Engine/Plugins/Experimental/Animation/SkeletalMeshMorphTargetEditingTools/Source/SkeletalMeshMorphTargetEditingTools/Public/SkeletalMeshMorphTargetEditingToolsModule.h b/Engine/Plugins/Experimental/Animation/SkeletalMeshMorphTargetEditingTools/Source/SkeletalMeshMorphTargetEditingTools/Public/SkeletalMeshMorphTargetEditingToolsModule.h
--- a/Engine/Plugins/Experimental/Animation/SkeletalMeshMorphTargetEditingTools/Source/SkeletalMeshMorphTargetEditingTools/Public/SkeletalMeshMorphTargetEditingToolsModule.h
+++ b/Engine/Plugins/Experimental/Animation/SkeletalMeshMorphTargetEditingTools/Source/SkeletalMeshMorphTargetEditingTools/Public/SkeletalMeshMorphTargetEditingToolsModule.h
@@ -15,12 +15,14 @@ public:
 	virtual void ShutdownModule() override;
 	
 	// ISkeletalMeshModelingModeToolExtension
-	FText GetExtensionName() override { return FText::FromString(TEXT("SkeletalMeshMorphTargetEditingTools"));}
+	virtual FText GetExtensionName() override { return FText::FromString(TEXT("SkeletalMeshMorphTargetEditingTools"));}
 
-	FText GetToolSectionName() override { return FText::FromString(TEXT("Morph"));};
+	virtual FText GetToolSectionName() override { return FText::FromString(TEXT("Morph"));};
 
-	void GetExtensionTools(const FExtensionToolQueryInfo& QueryInfo, TArray<FExtensionToolDescription>& OutTools) override;
+	virtual void GetExtensionTools(const FExtensionToolQueryInfo& QueryInfo, TArray<FExtensionToolDescription>& OutTools) override;
 	
-	bool GetExtensionExtendedInfo(FModelingModeExtensionExtendedInfo& InfoOut) override;
+	virtual bool GetExtensionExtendedInfo(FModelingModeExtensionExtendedInfo& InfoOut) override;
+
+	virtual bool GetExtensionToolTargets(TArray<TSubclassOf<UToolTargetFactory>>& ToolTargetFactoriesOut) override;
 
 };
diff --git a/Engine/Plugins/Experimental/Animation/SkeletalMeshMorphTargetEditingTools/Source/SkeletalMeshMorphTargetEditingTools/Public/SkeletalMeshMorphTargetToolTarget.h b/Engine/Plugins/Experimental/Animation/SkeletalMeshMorphTargetEditingTools/Source/SkeletalMeshMorphTargetEditingTools/Public/SkeletalMeshMorphTargetToolTarget.h
new file mode 100644
--- /dev/null
+++ b/Engine/Plugins/Experimental/Animation/SkeletalMeshMorphTargetEditingTools/Source/SkeletalMeshMorphTargetEditingTools/Public/SkeletalMeshMorphTargetToolTarget.h
@@ -0,0 +1,68 @@
+﻿// Copyright Epic Games, Inc. All Rights Reserved.
+
+#pragma once
+
+#include "CoreMinimal.h"
+
+#include "ToolTargets/SkeletalMeshComponentToolTarget.h"
+#include "SkeletalMeshMorphTargetBackedTarget.h"
+#include "ToolTargets/ToolTarget.h"
+
+#include "SkeletalMeshMorphTargetToolTarget.generated.h"
+
+/**
+ * A tool target backed by a skeletal mesh component that can provide and take a mesh
+ * description.
+ */
+UCLASS(Transient)
+class USkeletalMeshMorphTargetToolTarget :
+	public USkeletalMeshComponentReadOnlyToolTarget,
+	public IDynamicMeshCommitter,
+	public ISkeletalMeshMorphTargetBackedTarget
+{
+	GENERATED_BODY()
+
+public:
+	virtual UE::Geometry::FDynamicMesh3 GetDynamicMesh() override;
+	virtual UE::Geometry::FDynamicMesh3 GetDynamicMesh(const FGetMeshParameters& InGetMeshParams) override;
+	
+	// IDynamicMeshCommitter
+	virtual void CommitDynamicMesh(const UE::Geometry::FDynamicMesh3& Mesh, const FDynamicMeshCommitInfo& CommitInfo) override;
+	using IDynamicMeshCommitter::CommitDynamicMesh; // unhide the other overload
+
+
+	virtual void SetDataToCommit(TArray<FVector3f>* InDataPtr) override;
+	virtual void SetEditingMorphTargetName(const FName& InName) override;
+	virtual TArray<FName> GetEditableMorphTargetNames() override;
+	virtual FName GetValidNameForNewMorphTarget(const FName& InName) override;
+protected:
+	// So that the tool target factory can poke into Component.
+	friend class USkeletalMeshMorphTargetToolTargetFactory;
+
+	FMeshDescription* GetEditingMeshDescription();
+	void CommitEditedMeshDescription();
+	
+	UPROPERTY()
+	FName EditingMorphTargetName;
+
+	static constexpr int32 LodIndex0 = 0;
+
+	TArray<FVector3f>* DataPtr = nullptr;
+};
+
+
+
+
+UCLASS(Transient)
+class USkeletalMeshMorphTargetToolTargetFactory : public UToolTargetFactory
+{
+	GENERATED_BODY()
+
+public:
+
+	virtual bool CanBuildTarget(UObject* SourceObject, const FToolTargetTypeRequirements& Requirements) const override;
+
+	virtual UToolTarget* BuildTarget(UObject* SourceObject, const FToolTargetTypeRequirements& TargetTypeInfo) override;
+};
+
+
diff --git a/Engine/Plugins/Experimental/Animation/SkeletalMeshMorphTargetEditingTools/Source/SkeletalMeshMorphTargetEditingTools/SkeletalMeshMorphTargetEditingTools.Build.cs b/Engine/Plugins/Experimental/Animation/SkeletalMeshMorphTargetEditingTools/Source/SkeletalMeshMorphTargetEditingTools/SkeletalMeshMorphTargetEditingTools.Build.cs
--- a/Engine/Plugins/Experimental/Animation/SkeletalMeshMorphTargetEditingTools/Source/SkeletalMeshMorphTargetEditingTools/SkeletalMeshMorphTargetEditingTools.Build.cs
+++ b/Engine/Plugins/Experimental/Animation/SkeletalMeshMorphTargetEditingTools/Source/SkeletalMeshMorphTargetEditingTools/SkeletalMeshMorphTargetEditingTools.Build.cs
@@ -50,6 +50,11 @@ public class SkeletalMeshMorphTargetEditingTools : ModuleRules
 				"SkeletalMeshModifiers",
 				"MeshDescription",
 				"StaticMeshDescription",
+				"SkeletalMeshDescription",
+				"ModelingComponentsEditorOnly",
+				"UnrealEd",
+				"Persona",
+				"GeometryFramework"
 				// ... add private dependencies that you statically link with here ...	
 			}
 			);
diff --git a/Engine/Plugins/Runtime/MeshModelingToolset/Source/MeshModelingToolsEditorOnly/Private/SkeletalMesh/SkeletalMeshEditionInterface.cpp b/Engine/Plugins/Runtime/MeshModelingToolset/Source/MeshModelingToolsEditorOnly/Private/SkeletalMesh/SkeletalMeshEditionInterface.cpp
--- a/Engine/Plugins/Runtime/MeshModelingToolset/Source/MeshModelingToolsEditorOnly/Private/SkeletalMesh/SkeletalMeshEditionInterface.cpp
+++ b/Engine/Plugins/Runtime/MeshModelingToolset/Source/MeshModelingToolsEditorOnly/Private/SkeletalMesh/SkeletalMeshEditionInterface.cpp
@@ -50,6 +50,12 @@ TArray<FName> ISkeletalMeshEditingInterface::GetSelectedBones() const
 	return Binding.IsValid() ? Binding.Pin()->GetSelectedBones() : Dummy;
 }
 
+TArray<FName> ISkeletalMeshEditingInterface::GetSelectedMorphTargets() const
+{
+	static const TArray<FName> Dummy;
+	return Binding.IsValid() ? Binding.Pin()->GetSelectedMorphTargets() : Dummy;
+}
+
 FSkeletalMeshToolNotifier::FSkeletalMeshToolNotifier(TWeakInterfacePtr<ISkeletalMeshEditingInterface> InInterface)
 	: ISkeletalMeshNotifier()
 	, Interface(InInterface)
diff --git a/Engine/Plugins/Runtime/MeshModelingToolset/Source/MeshModelingToolsEditorOnly/Public/SkeletalMesh/SkeletalMeshEditionInterface.h b/Engine/Plugins/Runtime/MeshModelingToolset/Source/MeshModelingToolsEditorOnly/Public/SkeletalMesh/SkeletalMeshEditionInterface.h
--- a/Engine/Plugins/Runtime/MeshModelingToolset/Source/MeshModelingToolsEditorOnly/Public/SkeletalMesh/SkeletalMeshEditionInterface.h
+++ b/Engine/Plugins/Runtime/MeshModelingToolset/Source/MeshModelingToolsEditorOnly/Public/SkeletalMesh/SkeletalMeshEditionInterface.h
@@ -42,6 +42,7 @@ public:
 	bool NeedsNotification() const;
 
 	virtual TArray<FName> GetSelectedBones() const;
+	virtual TArray<FName> GetSelectedMorphTargets() const;
 	
 	void BindTo(TSharedPtr<ISkeletalMeshEditorBinding> InBinding);
 	void Unbind();
diff --git a/Engine/Source/Developer/SkeletalMeshUtilitiesCommon/Public/SkeletalMeshNotifier.h b/Engine/Source/Developer/SkeletalMeshUtilitiesCommon/Public/SkeletalMeshNotifier.h
--- a/Engine/Source/Developer/SkeletalMeshUtilitiesCommon/Public/SkeletalMeshNotifier.h
+++ b/Engine/Source/Developer/SkeletalMeshUtilitiesCommon/Public/SkeletalMeshNotifier.h
@@ -53,4 +53,7 @@ public:
 	virtual NameFunction GetNameFunction() = 0;
 
 	virtual TArray<FName> GetSelectedBones() const = 0;
+	virtual TArray<FName> GetSelectedMorphTargets() const = 0;
+	
+
 };
\ No newline at end of file
diff --git a/Engine/Source/Editor/Persona/Private/PersonaModule.cpp b/Engine/Source/Editor/Persona/Private/PersonaModule.cpp
--- a/Engine/Source/Editor/Persona/Private/PersonaModule.cpp
+++ b/Engine/Source/Editor/Persona/Private/PersonaModule.cpp
@@ -104,6 +104,7 @@
 #include "SAssetView.h"
 #include "ToolMenus.h"
 #include "SAnimAssetFindReplace.h"
+#include "SMorphTargetViewer.h"
 
 IMPLEMENT_MODULE( FPersonaModule, Persona );
 
@@ -319,7 +320,21 @@ TSharedRef<class FWorkflowTabFactory> FPersonaModule::CreateAssetDetailsTabFacto
 
 TSharedRef<class FWorkflowTabFactory> FPersonaModule::CreateMorphTargetTabFactory(const TSharedRef<class FWorkflowCentricApplication>& InHostingApp, const TSharedRef<IPersonaPreviewScene>& InPreviewScene, FSimpleMulticastDelegate& OnPostUndo) const
 {
-	return MakeShareable(new FMorphTargetTabSummoner(InHostingApp, InPreviewScene, OnPostUndo));
+	TSharedRef<IMorphTargetViewer> ViewerWidget = CreateDefaultMorphTargetViewerWidget(InPreviewScene, OnPostUndo);
+	return MakeShareable(new FMorphTargetTabSummoner(InHostingApp, ViewerWidget));
+}
+
+TSharedRef<FWorkflowTabFactory> FPersonaModule::CreateMorphTargetTabFactory(
+	const TSharedRef<class FWorkflowCentricApplication>& InHostingApp, TSharedRef<IMorphTargetViewer> InViewerWidget) const
+{
+	return MakeShareable(new FMorphTargetTabSummoner(InHostingApp, InViewerWidget));
+}
+
+TSharedRef<IMorphTargetViewer> FPersonaModule::CreateDefaultMorphTargetViewerWidget(
+	const TSharedRef<IPersonaPreviewScene>& InPreviewScene,
+	FSimpleMulticastDelegate& OnPostUndo) const
+{
+	return SNew(SMorphTargetViewer, InPreviewScene, OnPostUndo);	
 }
 
 TSharedRef<class FWorkflowTabFactory> FPersonaModule::CreateAnimBlueprintPreviewTabFactory(const TSharedRef<class FBlueprintEditor>& InBlueprintEditor, const TSharedRef<IPersonaPreviewScene>& InPreviewScene) const
diff --git a/Engine/Source/Editor/Persona/Private/SMorphTargetViewer.cpp b/Engine/Source/Editor/Persona/Private/SMorphTargetViewer.cpp
--- a/Engine/Source/Editor/Persona/Private/SMorphTargetViewer.cpp
+++ b/Engine/Source/Editor/Persona/Private/SMorphTargetViewer.cpp
@@ -25,6 +25,7 @@
 #include "Widgets/Input/SCheckBox.h"
 #include "Widgets/Input/SSearchBox.h"
 #include "Widgets/Input/SSpinBox.h"
+#include "Editor.h"
 
 #define LOCTEXT_NAMESPACE "SMorphTargetViewer"
 
@@ -64,6 +65,12 @@ public:
 
 private:
 
+	/**
+	* Called when the user begins/ends dragging the slider on the SSpinBox
+	*/
+	void OnBeginSlideMorphTargetWeight();
+	void OnEndSlideMorphTargetWeight(float Value);
+	
 	/**
 	* Called when the user changes the value of the SSpinBox
 	*
@@ -220,6 +227,8 @@ TSharedRef< SWidget > SMorphTargetListRow::GenerateWidgetForColumn( const FName&
 				.MinSliderValue(-1.f)
 				.MaxSliderValue(1.f)
 				.Value( this, &SMorphTargetListRow::GetWeight )
+				.OnBeginSliderMovement(this, &SMorphTargetListRow::OnBeginSlideMorphTargetWeight)
+				.OnEndSliderMovement(this, &SMorphTargetListRow::OnEndSlideMorphTargetWeight)
 				.OnValueChanged( this, &SMorphTargetListRow::OnMorphTargetWeightChanged )
 				.OnValueCommitted( this, &SMorphTargetListRow::OnMorphTargetWeightValueCommitted )
 				.IsEnabled(this, &SMorphTargetListRow::IsMorphTargetWeightSliderEnabled)
@@ -266,6 +275,16 @@ TSharedRef< SWidget > SMorphTargetListRow::GenerateWidgetForColumn( const FName&
 	}
 }
 
+void SMorphTargetListRow::OnBeginSlideMorphTargetWeight()
+{
+	GEditor->BeginTransaction(LOCTEXT("OverrideMorphTargetWeight", "Override Morph Target Weight"));
+}
+
+void SMorphTargetListRow::OnEndSlideMorphTargetWeight(float Value)
+{
+	GEditor->EndTransaction();
+}
+
 void SMorphTargetListRow::OnMorphTargetAutoFillChecked(ECheckBoxState InState)
 {
 	Item->bAutoFillData = InState == ECheckBoxState::Checked;
@@ -392,7 +411,8 @@ float SMorphTargetListRow::GetWeight() const
 	}
 	else
 	{
-		return Item->Weight;
+		USkeletalMeshComponent* SkelComp = PreviewScenePtr.Pin()->GetPreviewMeshComponent();
+		return SkelComp->GetMorphTargetCurves()[Item->Name];
 	}
 }
 //////////////////////////////////////////////////////////////////////////
@@ -484,6 +504,25 @@ void SMorphTargetViewer::Construct(const FArguments& InArgs, const TSharedRef<IP
 	CreateMorphTargetList();
 }
 
+TArray<FName> SMorphTargetViewer::GetSelectedMorphTargetNames() const
+{
+	TArray< TSharedPtr< FDisplayedMorphTargetInfo > > SelectedRows = MorphTargetListView->GetSelectedItems();
+
+	TArray<FName> SelectedMorphTargetNames;
+	for (auto ItemIt = SelectedRows.CreateIterator(); ItemIt; ++ItemIt)
+	{
+		TSharedPtr< FDisplayedMorphTargetInfo > RowItem = (*ItemIt);
+		SelectedMorphTargetNames.AddUnique(RowItem->Name);
+	}
+
+	return SelectedMorphTargetNames;
+}
+
+TSharedRef<SWidget> SMorphTargetViewer::AsWidget()
+{
+	return this->AsShared();
+}
+
 void SMorphTargetViewer::OnPreviewMeshChanged(class USkeletalMesh* OldPreviewMesh, class USkeletalMesh* NewPreviewMesh)
 {
 	SkeletalMesh = NewPreviewMesh;
@@ -674,6 +713,8 @@ TSharedPtr<SWidget> SMorphTargetViewer::OnGetContextMenuContent() const
 
 void SMorphTargetViewer::CreateMorphTargetList( const FString& SearchText )
 {
+	TArray<FName> SelectedMorphTargets = GetSelectedMorphTargetNames();
+	
 	MorphTargetList.Empty();
 
 	if ( SkeletalMesh )
@@ -721,6 +762,7 @@ void SMorphTargetViewer::CreateMorphTargetList( const FString& SearchText )
 				if(CurveValPtr)
 				{
 					Info.Get().Weight = (*CurveValPtr);
+					Info.Get().bAutoFillData = false;
 				}
 			}
 
@@ -728,6 +770,15 @@ void SMorphTargetViewer::CreateMorphTargetList( const FString& SearchText )
 		}
 	}
 
+	for (const FDisplayedMorphTargetInfoPtr& Item : MorphTargetList)
+	{
+		if (SelectedMorphTargets.Contains(Item->Name))
+		{
+			MorphTargetListView->SetItemSelection(Item, true);
+		}
+	}
+	
+	
 	NotifySelectionChange();
 	MorphTargetListView->RequestListRefresh();
 }
@@ -738,7 +789,11 @@ void SMorphTargetViewer::AddMorphTargetOverride( FName& Name, float Weight, bool
 
 	if ( Mesh )
 	{
+		GEditor->BeginTransaction(LOCTEXT("MorphTargetOverrideChanged", "Changed Morph Target Override"));
+		Mesh->SetFlags(RF_Transactional);
+		Mesh->Modify();
 		Mesh->SetMorphTarget( Name, Weight, bRemoveZeroWeight );
+		GEditor->EndTransaction();
 	}
 }
 
@@ -941,17 +996,10 @@ void SMorphTargetViewer::OnMeshChanged()
 
 void SMorphTargetViewer::NotifySelectionChange() const
 {
-	TArray< TSharedPtr< FDisplayedMorphTargetInfo > > SelectedRows = MorphTargetListView->GetSelectedItems();
-
-	TArray<FName> SelectedMorphtargetNames;
-	for (auto ItemIt = SelectedRows.CreateIterator(); ItemIt; ++ItemIt)
-	{
-		TSharedPtr< FDisplayedMorphTargetInfo > RowItem = (*ItemIt);
-		SelectedMorphtargetNames.AddUnique(RowItem->Name);
-	}
+	TArray<FName> SelectedMorphTargetNames = GetSelectedMorphTargetNames();
 
 	// stil have to call this even if empty, otherwise it won't clear it
-	SetSelectedMorphTargets(SelectedMorphtargetNames);
+	PreviewMorphTargets(SelectedMorphTargetNames);
 }
 
 void SMorphTargetViewer::OnRowsSelectedChanged(TSharedPtr<FDisplayedMorphTargetInfo> Item, ESelectInfo::Type SelectInfo)
@@ -959,7 +1007,7 @@ void SMorphTargetViewer::OnRowsSelectedChanged(TSharedPtr<FDisplayedMorphTargetI
 	NotifySelectionChange();
 }
 
-void SMorphTargetViewer::SetSelectedMorphTargets(const TArray<FName>& SelectedMorphTargetNames) const
+void SMorphTargetViewer::PreviewMorphTargets(const TArray<FName>& SelectedMorphTargetNames) const
 {
 	UDebugSkelMeshComponent* PreviewComponent = PreviewScenePtr.Pin()->GetPreviewMeshComponent();
 	if (PreviewComponent)
diff --git a/Engine/Source/Editor/Persona/Private/SMorphTargetViewer.h b/Engine/Source/Editor/Persona/Private/SMorphTargetViewer.h
--- a/Engine/Source/Editor/Persona/Private/SMorphTargetViewer.h
+++ b/Engine/Source/Editor/Persona/Private/SMorphTargetViewer.h
@@ -3,6 +3,7 @@
 #pragma once
 
 #include "CoreMinimal.h"
+#include "IMorphTargetViewer.h"
 #include "SlateFwd.h"
 #include "Widgets/DeclarativeSyntaxSupport.h"
 #include "Widgets/SCompoundWidget.h"
@@ -58,7 +59,7 @@ typedef SListView< TSharedPtr<FDisplayedMorphTargetInfo> > SMorphTargetListType;
 //////////////////////////////////////////////////////////////////////////
 // SMorphTargetViewer
 
-class SMorphTargetViewer : public SCompoundWidget
+class SMorphTargetViewer : public SCompoundWidget, public IMorphTargetViewer
 {
 public:
 	SLATE_BEGIN_ARGS( SMorphTargetViewer )
@@ -74,6 +75,10 @@ public:
 	*/
 	void Construct( const FArguments& InArgs, const TSharedRef<class IPersonaPreviewScene>& InPreviewScene, FSimpleMulticastDelegate& OnPostUndo );
 
+	/** IMorphTargetViewer overrides */
+	virtual TArray<FName> GetSelectedMorphTargetNames() const override;
+	virtual TSharedRef<SWidget> AsWidget() override;
+
 	/**
 	* Destructor - resets the morph targets
 	*
@@ -197,7 +202,7 @@ private:
 	void CreateMorphTargetList( const FString& SearchText = FString() );
 
 	// Sets the selected morph target
-	void SetSelectedMorphTargets(const TArray<FName>& SelectedMorphTargetNames) const;
+	void PreviewMorphTargets(const TArray<FName>& SelectedMorphTargetNames) const;
 
 	/** Pointer back to the Persona that owns us */
 	TWeakPtr<class IPersonaPreviewScene> PreviewScenePtr;
diff --git a/Engine/Source/Editor/Persona/Private/TabSpawners.cpp b/Engine/Source/Editor/Persona/Private/TabSpawners.cpp
--- a/Engine/Source/Editor/Persona/Private/TabSpawners.cpp
+++ b/Engine/Source/Editor/Persona/Private/TabSpawners.cpp
@@ -150,10 +150,9 @@ void FPersonaModeSharedData::Restore(const TSharedRef<FAnimationViewportClient>&
 /////////////////////////////////////////////////////
 // FMorphTargetTabSummoner
 
-FMorphTargetTabSummoner::FMorphTargetTabSummoner(TSharedPtr<class FAssetEditorToolkit> InHostingApp, const TSharedRef<IPersonaPreviewScene>& InPreviewScene, FSimpleMulticastDelegate& InOnPostUndo)
+FMorphTargetTabSummoner::FMorphTargetTabSummoner(TSharedPtr<class FAssetEditorToolkit> InHostingApp, TSharedRef<IMorphTargetViewer> InViewerWidget)
 	: FWorkflowTabFactory(FPersonaTabs::MorphTargetsID, InHostingApp)
-	, PreviewScene(InPreviewScene)
-	, OnPostUndo(InOnPostUndo)
+	, ViewerWidget(InViewerWidget)
 {
 	TabLabel = LOCTEXT("MorphTargetTabTitle", "Morph Target Previewer");
 	TabIcon = FSlateIcon(FAppStyle::GetAppStyleSetName(), "Persona.Tabs.MorphTargetPreviewer");
@@ -167,7 +166,7 @@ FMorphTargetTabSummoner::FMorphTargetTabSummoner(TSharedPtr<class FAssetEditorTo
 
 TSharedRef<SWidget> FMorphTargetTabSummoner::CreateTabBody(const FWorkflowTabSpawnInfo& Info) const
 {
-	return SNew(SMorphTargetViewer, PreviewScene.Pin().ToSharedRef(), OnPostUndo);
+	return ViewerWidget->AsWidget();
 }
 /////////////////////////////////////////////////////
 // FAnimCurveViewerTabSummoner
diff --git a/Engine/Source/Editor/Persona/Private/TabSpawners.h b/Engine/Source/Editor/Persona/Private/TabSpawners.h
--- a/Engine/Source/Editor/Persona/Private/TabSpawners.h
+++ b/Engine/Source/Editor/Persona/Private/TabSpawners.h
@@ -123,7 +123,7 @@ protected:
 struct FMorphTargetTabSummoner : public FWorkflowTabFactory
 {
 public:
-	FMorphTargetTabSummoner(TSharedPtr<class FAssetEditorToolkit> InHostingApp, const TSharedRef<IPersonaPreviewScene>& InPreviewScene, FSimpleMulticastDelegate& InOnPostUndo);
+	FMorphTargetTabSummoner(TSharedPtr<class FAssetEditorToolkit> InHostingApp, TSharedRef<IMorphTargetViewer> InViewerWidget);
 
 	virtual TSharedRef<SWidget> CreateTabBody(const FWorkflowTabSpawnInfo& Info) const override;
 
@@ -134,8 +134,7 @@ public:
 	}
 
 private:
-	TWeakPtr<class IPersonaPreviewScene> PreviewScene;
-	FSimpleMulticastDelegate& OnPostUndo;
+	TSharedRef<IMorphTargetViewer> ViewerWidget;
 };
 
 /////////////////////////////////////////////////////
diff --git a/Engine/Source/Editor/Persona/Public/IMorphTargetViewer.h b/Engine/Source/Editor/Persona/Public/IMorphTargetViewer.h
new file mode 100644
--- /dev/null
+++ b/Engine/Source/Editor/Persona/Public/IMorphTargetViewer.h
@@ -0,0 +1,17 @@
+// Copyright Epic Games, Inc. All Rights Reserved.
+
+#pragma once
+
+#include "CoreMinimal.h"
+#include "Widgets/SWidget.h"
+
+/** Interface to the morph target viewer in the skeletal mesh editor */
+class IMorphTargetViewer
+{
+public:
+	
+	virtual TArray<FName> GetSelectedMorphTargetNames() const = 0;
+	virtual TSharedRef<SWidget> AsWidget() = 0;
+
+	virtual ~IMorphTargetViewer() = default;
+};
\ No newline at end of file
diff --git a/Engine/Source/Editor/Persona/Public/PersonaModule.h b/Engine/Source/Editor/Persona/Public/PersonaModule.h
--- a/Engine/Source/Editor/Persona/Public/PersonaModule.h
+++ b/Engine/Source/Editor/Persona/Public/PersonaModule.h
@@ -13,6 +13,7 @@
 #include "SPersonaToolBox.h"
 #include "IBlendProfilePickerExtender.h"
 
+class IMorphTargetViewer;
 class FBlueprintEditor;
 class IDetailsView;
 class IEditableSkeleton;
@@ -384,6 +385,12 @@ public:
 	/** Create a tab factory for for previewing morph targets */
 	virtual TSharedRef<FWorkflowTabFactory> CreateMorphTargetTabFactory(const TSharedRef<class FWorkflowCentricApplication>& InHostingApp, const TSharedRef<IPersonaPreviewScene>& InPreviewScene, FSimpleMulticastDelegate& OnPostUndo) const;
 
+	/** Create a tab factory for for previewing morph targets */
+	virtual TSharedRef<FWorkflowTabFactory> CreateMorphTargetTabFactory(const TSharedRef<class FWorkflowCentricApplication>& InHostingApp, TSharedRef<IMorphTargetViewer> InViewerWidget) const;
+
+	/** Create the default morph target viewer widget*/
+	virtual TSharedRef<IMorphTargetViewer> CreateDefaultMorphTargetViewerWidget(const TSharedRef<IPersonaPreviewScene>& InPreviewScene, FSimpleMulticastDelegate& OnPostUndo) const;
+	
 	/** Create a tab factory for editing anim blueprint preview & defaults */
 	virtual TSharedRef<FWorkflowTabFactory> CreateAnimBlueprintPreviewTabFactory(const TSharedRef<class FBlueprintEditor>& InBlueprintEditor, const TSharedRef<IPersonaPreviewScene>& InPreviewScene) const;
 
diff --git a/Engine/Source/Editor/SkeletalMeshEditor/Private/SkeletalMeshEditor.cpp b/Engine/Source/Editor/SkeletalMeshEditor/Private/SkeletalMeshEditor.cpp
--- a/Engine/Source/Editor/SkeletalMeshEditor/Private/SkeletalMeshEditor.cpp
+++ b/Engine/Source/Editor/SkeletalMeshEditor/Private/SkeletalMeshEditor.cpp
@@ -62,6 +62,7 @@
 #include "AssetToolsModule.h"
 #include "DetailLayoutBuilder.h"
 #include "IAssetTools.h"
+#include "IMorphTargetViewer.h"
 #include "SkeletalMeshEditorContextMenuContext.h"
 #include "Styling/AppStyle.h"
 #include "InterchangeAssetImportData.h"
@@ -248,6 +249,8 @@ void FSkeletalMeshEditor::InitSkeletalMeshEditor(const EToolkitMode::Type Mode,
 	ISkeletonEditorModule& SkeletonEditorModule = FModuleManager::GetModuleChecked<ISkeletonEditorModule>("SkeletonEditor");
 	SkeletonTree = SkeletonEditorModule.CreateSkeletonTree(PersonaToolkit->GetSkeleton(), SkeletonTreeArgs);
 
+	MorphTargetViewer = PersonaModule.CreateDefaultMorphTargetViewerWidget(PreviewScene.ToSharedRef(), OnPostUndo);
+	
 	const bool bCreateDefaultStandaloneMenu = true;
 	const bool bCreateDefaultToolbar = true;
 	FAssetEditorToolkit::InitAssetEditor(Mode, InitToolkitHost, SkeletalMeshEditorAppIdentifier, FTabManager::FLayout::NullLayout, bCreateDefaultStandaloneMenu, bCreateDefaultToolbar, InSkeletalMesh);
@@ -256,7 +259,7 @@ void FSkeletalMeshEditor::InitSkeletalMeshEditor(const EToolkitMode::Type Mode,
 
 	AddApplicationMode(
 		SkeletalMeshEditorModes::SkeletalMeshEditorMode,
-		MakeShareable(new FSkeletalMeshEditorMode(SharedThis(this), SkeletonTree.ToSharedRef())));
+		MakeShareable(new FSkeletalMeshEditorMode(SharedThis(this), SkeletonTree.ToSharedRef(), MorphTargetViewer.ToSharedRef())));
 
 	SetCurrentMode(SkeletalMeshEditorModes::SkeletalMeshEditorMode);
 
@@ -1576,6 +1579,18 @@ TArray<FName> FSkeletalMeshEditorBinding::GetSelectedBones() const
 	return Selection;
 }
 
+TArray<FName> FSkeletalMeshEditorBinding::GetSelectedMorphTargets() const
+{
+	TArray<FName> Selection;
+
+	if (!Editor.IsValid())
+	{
+		return Selection;
+	}
+
+	return Editor.Pin()->GetMorphTargetViewer()->GetSelectedMorphTargetNames();
+}
+
 FSkeletalMeshEditorNotifier::FSkeletalMeshEditorNotifier(TSharedRef<FSkeletalMeshEditor> InEditor)
 	: ISkeletalMeshNotifier()
 	, Editor(InEditor)
diff --git a/Engine/Source/Editor/SkeletalMeshEditor/Private/SkeletalMeshEditor.h b/Engine/Source/Editor/SkeletalMeshEditor/Private/SkeletalMeshEditor.h
--- a/Engine/Source/Editor/SkeletalMeshEditor/Private/SkeletalMeshEditor.h
+++ b/Engine/Source/Editor/SkeletalMeshEditor/Private/SkeletalMeshEditor.h
@@ -107,6 +107,9 @@ public:
 	/** Get the skeleton tree widget */
 	TSharedRef<class ISkeletonTree> GetSkeletonTree() const { return SkeletonTree.ToSharedRef(); }
 
+	/** Get the MorphTargetViewer widget */
+	TSharedRef<class IMorphTargetViewer> GetMorphTargetViewer() const {return MorphTargetViewer.ToSharedRef(); }
+
 	void HandleDetailsCreated(const TSharedRef<class IDetailsView>& InDetailsView);
 
 	void HandleMeshDetailsCreated(const TSharedRef<class IDetailsView>& InDetailsView);
@@ -209,6 +212,9 @@ private:
 	/** Skeleton tree */
 	TSharedPtr<class ISkeletonTree> SkeletonTree;
 
+	/** Morph target viewer */
+	TSharedPtr<class IMorphTargetViewer> MorphTargetViewer;
+
 	/** Viewport */
 	TSharedPtr<class IPersonaViewport> Viewport;
 
@@ -251,7 +257,9 @@ public:
 	virtual ISkeletalMeshNotifier& GetNotifier() override;
 	virtual NameFunction GetNameFunction() override;
 	virtual TArray<FName> GetSelectedBones() const override;
-	
+	virtual TArray<FName> GetSelectedMorphTargets() const override;
+
+
 private:
 	TWeakPtr<FSkeletalMeshEditor> Editor;
 	FSkeletalMeshEditorNotifier Notifier;
diff --git a/Engine/Source/Editor/SkeletalMeshEditor/Private/SkeletalMeshEditorMode.cpp b/Engine/Source/Editor/SkeletalMeshEditor/Private/SkeletalMeshEditorMode.cpp
--- a/Engine/Source/Editor/SkeletalMeshEditor/Private/SkeletalMeshEditorMode.cpp
+++ b/Engine/Source/Editor/SkeletalMeshEditor/Private/SkeletalMeshEditorMode.cpp
@@ -5,6 +5,7 @@
 #include "PersonaModule.h"
 #include "SkeletalMeshEditor.h"
 #include "ISkeletonTree.h"
+#include "IMorphTargetViewer.h"
 #include "ISkeletonEditorModule.h"
 #include "IPersonaToolkit.h"
 #include "SControlRigMappingWindow.h"
@@ -15,7 +16,11 @@
 
 #define LOCTEXT_NAMESPACE "SkeletalMeshEditorMode"
 
-FSkeletalMeshEditorMode::FSkeletalMeshEditorMode(TSharedRef<FWorkflowCentricApplication> InHostingApp, TSharedRef<ISkeletonTree> InSkeletonTree)
+FSkeletalMeshEditorMode::FSkeletalMeshEditorMode(
+	TSharedRef<FWorkflowCentricApplication> InHostingApp,
+	TSharedRef<ISkeletonTree> InSkeletonTree,
+	TSharedRef<IMorphTargetViewer> InMorphTargetViewer
+	)
 	: FApplicationMode(SkeletalMeshEditorModes::SkeletalMeshEditorMode)
 {
 	HostingAppPtr = InHostingApp;
@@ -38,7 +43,7 @@ FSkeletalMeshEditorMode::FSkeletalMeshEditorMode(TSharedRef<FWorkflowCentricAppl
 
 	TabFactories.RegisterFactory(PersonaModule.CreateAdvancedPreviewSceneTabFactory(InHostingApp, SkeletalMeshEditor->GetPersonaToolkit()->GetPreviewScene()));
 	TabFactories.RegisterFactory(PersonaModule.CreateAssetDetailsTabFactory(InHostingApp, FOnGetAsset::CreateSP(&SkeletalMeshEditor.Get(), &FSkeletalMeshEditor::HandleGetAsset), FOnDetailsCreated::CreateSP(&SkeletalMeshEditor.Get(), &FSkeletalMeshEditor::HandleMeshDetailsCreated)));
-	TabFactories.RegisterFactory(PersonaModule.CreateMorphTargetTabFactory(InHostingApp, SkeletalMeshEditor->GetPersonaToolkit()->GetPreviewScene(), SkeletalMeshEditor->OnPostUndo));
+	TabFactories.RegisterFactory(PersonaModule.CreateMorphTargetTabFactory(InHostingApp, InMorphTargetViewer));
 	TabFactories.RegisterFactory(PersonaModule.CreateCurveMetadataEditorTabFactory(InHostingApp, SkeletalMeshEditor->HandleGetAsset(), SkeletalMeshEditor->GetPersonaToolkit()->GetPreviewScene(), OnObjectsSelected));
 	TabFactories.RegisterFactory(CreateMeshControllerMappingTabFactory(InHostingApp, Cast<USkeletalMesh> (SkeletalMeshEditor->HandleGetAsset()), SkeletalMeshEditor->OnPostUndo));
 	TabFactories.RegisterFactory(PersonaModule.CreateAnimAssetFindReplaceTabFactory(InHostingApp, FAnimAssetFindReplaceConfig()));
diff --git a/Engine/Source/Editor/SkeletalMeshEditor/Private/SkeletalMeshEditorMode.h b/Engine/Source/Editor/SkeletalMeshEditor/Private/SkeletalMeshEditorMode.h
--- a/Engine/Source/Editor/SkeletalMeshEditor/Private/SkeletalMeshEditorMode.h
+++ b/Engine/Source/Editor/SkeletalMeshEditor/Private/SkeletalMeshEditorMode.h
@@ -17,7 +17,11 @@ class SToolTip;
 class FSkeletalMeshEditorMode : public FApplicationMode
 {
 public:
-	FSkeletalMeshEditorMode(TSharedRef<class FWorkflowCentricApplication> InHostingApp, TSharedRef<class ISkeletonTree> InSkeletonTree);
+	FSkeletalMeshEditorMode(
+		TSharedRef<FWorkflowCentricApplication> InHostingApp,
+		TSharedRef<class ISkeletonTree> InSkeletonTree,
+		TSharedRef<class IMorphTargetViewer> InMorphTargetViewer
+		);
 
 	/** FApplicationMode interface */
 	virtual void RegisterTabFactories(TSharedPtr<FTabManager> InTabManager) override;
diff --git a/Engine/Source/Runtime/Engine/Classes/Components/SkeletalMeshComponent.h b/Engine/Source/Runtime/Engine/Classes/Components/SkeletalMeshComponent.h
--- a/Engine/Source/Runtime/Engine/Classes/Components/SkeletalMeshComponent.h
+++ b/Engine/Source/Runtime/Engine/Classes/Components/SkeletalMeshComponent.h
@@ -1645,6 +1645,7 @@ private:
 	 * Morph Target Curves. This will override AnimInstance MorphTargetCurves
 	 * if same curve is found
 	 **/
+	UPROPERTY(Transient)
 	TMap<FName, float>	MorphTargetCurves;
 
 public:
diff --git a/Engine/Source/Runtime/SkeletalMeshDescription/Private/SkeletalMeshOperations.cpp b/Engine/Source/Runtime/SkeletalMeshDescription/Private/SkeletalMeshOperations.cpp
--- a/Engine/Source/Runtime/SkeletalMeshDescription/Private/SkeletalMeshOperations.cpp
+++ b/Engine/Source/Runtime/SkeletalMeshDescription/Private/SkeletalMeshOperations.cpp
@@ -485,6 +485,162 @@ static void PoseMesh(
 		});
 }
 
+static void GetMorphTargetDeltasFromPosedMesh(
+	FMeshDescription& InOutTargetMesh,
+	const FMeshDescription& InRefMesh,
+	TConstArrayView<FMatrix44f> InRefToUserTransforms,
+	const FName InSkinWeightProfile,
+	const TMap<FName, float>& InMorphTargetWeights
+)
+{
+	struct FMorphInfo
+	{
+		TVertexAttributesConstRef<FVector3f> PositionDelta;
+		TVertexInstanceAttributesConstRef<FVector3f> NormalDelta;
+		float Weight = 0.0f;
+	};
+
+	FSkeletalMeshAttributes Attributes(InOutTargetMesh);
+	FSkeletalMeshConstAttributes RefAttributes(InRefMesh);
+
+	// We need the mesh to be compact for the parallel for to work.
+	FElementIDRemappings Remappings;
+	InOutTargetMesh.Compact(Remappings);
+
+	TVertexAttributesRef<FVector3f> PositionAttribute = Attributes.GetVertexPositions();
+	TVertexInstanceAttributesRef<FVector3f> NormalAttribute = Attributes.GetVertexInstanceNormals();
+	TVertexInstanceAttributesRef<FVector3f> TangentAttribute = Attributes.GetVertexInstanceTangents();
+	TVertexInstanceAttributesRef<float> BinormalSignsAttribute = Attributes.GetVertexInstanceBinormalSigns();
+
+
+	// Invert skinning first
+
+	using namespace UE::AnimationCore;
+
+	// The normal needs to be transformed using the inverse transpose of the transform matrices to ensure that
+	// scaling works correctly. 
+	TArray<FMatrix44f> RefToUserTransformsNormal;
+	RefToUserTransformsNormal.Reserve(InRefToUserTransforms.Num());
+	for (const FMatrix44f& Mat : InRefToUserTransforms)
+	{
+		RefToUserTransformsNormal.Add(Mat.Inverse().GetTransposed());
+	}
+
+	FSkinWeightsVertexAttributesRef SkinWeightAttribute = Attributes.GetVertexSkinWeights(InSkinWeightProfile);
+	ParallelFor(InOutTargetMesh.Vertices().Num(),
+		[&Mesh = InOutTargetMesh, &PositionAttribute, &NormalAttribute, &TangentAttribute, &SkinWeightAttribute, &RefToUserTransforms = InRefToUserTransforms, &RefToUserTransformsNormal](int32 Index)
+		{
+			const FVertexID VertexID(Index);
+			const FVertexBoneWeights BoneWeights = SkinWeightAttribute.Get(VertexID);
+			const FVector3f Position = PositionAttribute.Get(VertexID);
+			FVector3f SkinnedPosition = FVector3f::ZeroVector;
+
+			FMatrix44f SkinMatrix = FMatrix44f(FVector3f::ZeroVector, FVector3f::ZeroVector, FVector3f::ZeroVector, FVector3f::ZeroVector);
+			SkinMatrix.M[3][3] = 0.0f;
+
+			for (const FBoneWeight BW : BoneWeights)
+			{
+				SkinMatrix += RefToUserTransforms[BW.GetBoneIndex()] * BW.GetWeight();
+			}
+			SkinnedPosition = SkinMatrix.Inverse().TransformPosition(Position);
+			PositionAttribute.Set(VertexID, SkinnedPosition);
+
+			for (const FVertexInstanceID VertexInstanceID : Mesh.GetVertexVertexInstanceIDs(VertexID))
+			{
+				const FVector3f Normal = NormalAttribute.Get(VertexInstanceID);
+				const FVector3f Tangent = TangentAttribute.Get(VertexInstanceID);
+				FVector3f SkinnedNormal = FVector3f::ZeroVector;
+				FVector3f SkinnedTangent = FVector3f::ZeroVector;
+
+				for (const FBoneWeight BW : BoneWeights)
+				{
+					SkinnedNormal += RefToUserTransformsNormal[BW.GetBoneIndex()].TransformVector(Normal) * BW.GetWeight();
+					SkinnedTangent += RefToUserTransforms[BW.GetBoneIndex()].TransformVector(Tangent) * BW.GetWeight();
+				}
+
+				SkinnedNormal.Normalize();
+				SkinnedTangent.Normalize();
+
+				NormalAttribute.Set(VertexInstanceID, SkinnedNormal);
+				TangentAttribute.Set(VertexInstanceID, SkinnedTangent);
+			}
+		});
+
+
+	// See which morph target attributes we can peel out. If the normal attributes are not all valid, then
+	// we have to automatically compute the normal from the positions. Otherwise, we only use the normal deltas. 
+	TArray<FMorphInfo> MorphInfos;
+	bool bAllMorphNormalsValid = true;
+	for (const TPair<FName, float>& Item : InMorphTargetWeights)
+	{
+		const FName MorphName{ Item.Key };
+		const float MorphWeight{ Item.Value };
+		TVertexAttributesConstRef<FVector3f> PositionDelta = RefAttributes.GetVertexMorphPositionDelta(MorphName);
+		// Q: Should we use the value of `r.MorphTarget.WeightThreshold` instead? The following condition is
+		// identical to the default setting of that value.
+		if (PositionDelta.IsValid() && !FMath::IsNearlyZero(MorphWeight))
+		{
+			TVertexInstanceAttributesConstRef<FVector3f> NormalDelta = RefAttributes.GetVertexInstanceMorphNormalDelta(MorphName);
+			if (!NormalDelta.IsValid())
+			{
+				bAllMorphNormalsValid = false;
+			}
+			FMorphInfo MorphInfo;
+			MorphInfo.PositionDelta = PositionDelta;
+			MorphInfo.NormalDelta = NormalDelta;
+			MorphInfo.Weight = MorphWeight;
+			MorphInfos.Add(MorphInfo);
+		}
+	}
+
+	// Inverse morph deltas
+
+	if (!MorphInfos.IsEmpty())
+	{
+		struct FMorphProcessContext
+		{
+			TSet<FVertexInstanceID> DirtyVertexInstances;
+			TArray<FVertexID> Neighbors;
+		};
+		TArray<FMorphProcessContext> Contexts;
+		ParallelForWithTaskContext(Contexts, InOutTargetMesh.Vertices().Num(),
+			[&Mesh = InOutTargetMesh, &PositionAttribute, &MorphInfos, bAllMorphNormalsValid](FMorphProcessContext& Context, int32 Index)
+			{
+				const FVertexID VertexID{ Index };
+
+				FVector3f Position = PositionAttribute.Get(VertexID);
+				bool bMoved = false;
+				for (const FMorphInfo& MorphInfo : MorphInfos)
+				{
+					const FVector3f PositionDelta = MorphInfo.PositionDelta.Get(VertexID) * MorphInfo.Weight * -1.0;
+					if (!PositionDelta.IsNearlyZero())
+					{
+						Position += PositionDelta;
+						bMoved = true;
+					}
+				}
+
+				// If we need to re-generate the normals, store which vertices got moved _and_ their neighbors, since the whole
+				// triangle moved, which affects neighboring vertices of the moved vertex.
+				if (bMoved)
+				{
+					PositionAttribute.Set(VertexID, Position);
+
+					if (!bAllMorphNormalsValid)
+					{
+						Context.DirtyVertexInstances.Append(Mesh.GetVertexVertexInstanceIDs(VertexID));
+
+						Mesh.GetVertexAdjacentVertices(VertexID, Context.Neighbors);
+						for (const FVertexID NeighborVertexID : Context.Neighbors)
+						{
+							Context.DirtyVertexInstances.Append(Mesh.GetVertexVertexInstanceIDs(NeighborVertexID));
+						}
+					}
+				}
+			});
+	}
+}
+	
 } // namespace Impl
 
 bool FSkeletalMeshOperations::GetPosedMesh(
@@ -635,6 +791,75 @@ bool FSkeletalMeshOperations::GetPosedMesh(
 	return true;
 }
 
+bool FSkeletalMeshOperations::GetMorphTargetDeltasFromPosedMesh(
+	const FMeshDescription& InPoseMesh,
+	const FMeshDescription& InRefMesh,
+	TArray<FTransform>& RefBoneTransforms,
+	FMeshDescription& OutTargetMesh,
+	TConstArrayView<FTransform> InComponentSpaceTransforms,
+	const FName InSkinWeightProfile,
+	const TMap<FName, float>& InMorphTargetWeights
+)
+{
+	// Verify that the mesh is valid.
+	FSkeletalMeshConstAttributes Attributes(InPoseMesh);
+	if (!Attributes.HasBonePoseAttribute() || !Attributes.HasBoneParentIndexAttribute())
+	{
+		return false;
+	}
+
+	if (!Attributes.GetVertexSkinWeights(InSkinWeightProfile).IsValid())
+	{
+		return false;
+	}
+
+	if (Attributes.GetNumBones() != InComponentSpaceTransforms.Num())
+	{
+		return false;
+	}
+
+	// Convert the component-space transforms into a set of matrices that transform from the reference pose to
+	// the user pose. These are then used to nudge the vertices from the reference pose to the wanted
+	// user pose by weighing the influence of each bone on a given vertex. If the user pose and the reference pose
+	// are identical, these are all identity matrices.
+	TArray<FMatrix44f> RefToUserTransforms;
+	TArray<FMatrix44f> RefPoseTransforms;
+
+	FSkeletalMeshAttributes::FBonePoseAttributesConstRef BonePoseAttribute = Attributes.GetBonePoses();
+	FSkeletalMeshAttributes::FBoneParentIndexAttributesConstRef ParentBoneIndexAttribute = Attributes.GetBoneParentIndices();
+	RefToUserTransforms.SetNumUninitialized(Attributes.GetNumBones());
+	RefPoseTransforms.SetNumUninitialized(Attributes.GetNumBones());
+
+	for (int32 BoneIndex = 0; BoneIndex < Attributes.GetNumBones(); BoneIndex++)
+	{
+		const int32 ParentBoneIndex = ParentBoneIndexAttribute.Get(BoneIndex);
+		
+		FTransform ReferenceBoneTransform = RefBoneTransforms[BoneIndex];
+		RefPoseTransforms[BoneIndex] = FMatrix44f{ ReferenceBoneTransform.ToMatrixWithScale() };
+		RefToUserTransforms[BoneIndex] = RefPoseTransforms[BoneIndex].Inverse() * FMatrix44f { InComponentSpaceTransforms[BoneIndex].ToMatrixWithScale() };
+	}
+
+	// Start with a fresh duplicate and then pose the target mesh in-place.
+	OutTargetMesh = InPoseMesh;
+	UE::Impl::GetMorphTargetDeltasFromPosedMesh(OutTargetMesh, InRefMesh, RefToUserTransforms, InSkinWeightProfile, InMorphTargetWeights);
+
+	// Write out the current ref pose (in bone-space) to the mesh. 
+	FSkeletalMeshAttributes WriteAttributes(OutTargetMesh);
+	FSkeletalMeshAttributes::FBonePoseAttributesRef WriteBonePoseAttribute = WriteAttributes.GetBonePoses();
+	for (int32 BoneIndex = 0; BoneIndex < Attributes.GetNumBones(); BoneIndex++)
+	{
+		const int32 ParentBoneIndex = ParentBoneIndexAttribute.Get(BoneIndex);
+		FTransform RefPoseTransform = InComponentSpaceTransforms[BoneIndex];
+
+		if (ParentBoneIndex != INDEX_NONE)
+		{
+			RefPoseTransform = RefPoseTransform.GetRelativeTransform(InComponentSpaceTransforms[ParentBoneIndex]);
+		}
+		WriteBonePoseAttribute.Set(BoneIndex, RefPoseTransform);
+	}
+
+	return true;
+}
 
 void FSkeletalMeshOperations::ConvertHardEdgesToSmoothMasks(
 	const FMeshDescription& InMeshDescription,
diff --git a/Engine/Source/Runtime/SkeletalMeshDescription/Public/SkeletalMeshOperations.h b/Engine/Source/Runtime/SkeletalMeshDescription/Public/SkeletalMeshOperations.h
--- a/Engine/Source/Runtime/SkeletalMeshDescription/Public/SkeletalMeshOperations.h
+++ b/Engine/Source/Runtime/SkeletalMeshDescription/Public/SkeletalMeshOperations.h
@@ -105,7 +105,26 @@ public:
 		const TMap<FName, float>& InMorphTargetWeights = {}
 		);
 
-
+	/** Returns a Morph Target deltas extracted from the deformed mesh
+	 *  
+	 *  @param InPoseMesh Deformed mesh
+	 *  @param InRefMesh Ref pose mesh
+	 *  @param RefBoneTransforms Ref component space bone transforms
+	 *  @param OutTargetMesh Extracted morph target deltas 
+	 *  @param InComponentSpaceTransforms Current component space bone transforms
+	 *  @param InSkinWeightProfile The skin weight profile to use as the source of skin weights for the deformation.
+	 *  @param InMorphTargetWeights Active morph target weights producing the fully deformed mesh. Don't include the weight of the morph target you want to extract
+	 *  @return \c true if the operation succeeded.
+	 */
+	static SKELETALMESHDESCRIPTION_API bool GetMorphTargetDeltasFromPosedMesh(
+		const FMeshDescription& InPoseMesh,
+		const FMeshDescription& InRefMesh,
+		TArray<FTransform>& RefBoneTransforms,
+		FMeshDescription& OutTargetMesh,
+		TConstArrayView<FTransform> InComponentSpaceTransforms,
+		const FName InSkinWeightProfile,
+		const TMap<FName, float>& InMorphTargetWeights
+		);
 
 	/**
 	 * A simpler variant of FStaticMeshOperations::ConvertHardEdgesToSmoothGroup that assumes that hard edges always form closed regions. 
