[SkeletalMesh Morph Editing tool] fixed crash for meshes that contains non-manifold verts

diff --git a/Engine/Plugins/Experimental/Animation/SkeletalMeshMorphTargetEditingTools/Source/SkeletalMeshMorphTargetEditingTools/Private/MorphTargetVertexSculptTool.cpp b/Engine/Plugins/Experimental/Animation/SkeletalMeshMorphTargetEditingTools/Source/SkeletalMeshMorphTargetEditingTools/Private/MorphTargetVertexSculptTool.cpp
--- a/Engine/Plugins/Experimental/Animation/SkeletalMeshMorphTargetEditingTools/Source/SkeletalMeshMorphTargetEditingTools/Private/MorphTargetVertexSculptTool.cpp
+++ b/Engine/Plugins/Experimental/Animation/SkeletalMeshMorphTargetEditingTools/Source/SkeletalMeshMorphTargetEditingTools/Private/MorphTargetVertexSculptTool.cpp
@@ -221,7 +221,9 @@ void UMorphTargetVertexSculptTool::UpdateCacheIfNeeded()
 	
 	if (FSkeletalMeshOperations::GetMorphTargetDeltasFromPosedMesh(SculptMeshDescription, *RefMeshDescription, RefComponentSpaceTransforms, NonDeformedMesh, ComponentSpaceTransforms, NAME_None, MorphTargetWeights))
 	{
-		for (FVertexID VertexID : NonDeformedMesh.Vertices().GetElementIDs())
+		// Worth noting that the sculpt mesh description can have more verts due to the initial construction of the dynamic mesh by the tool target when
+		// the mesh has non-manifold structure, which is why we want to iterate over the RefMeshDescription
+		for (FVertexID VertexID : RefMeshDescription->Vertices().GetElementIDs())
 		{
 			FVector3f Delta = NonDeformedMesh.GetVertexPosition(VertexID) - RefMeshDescription->GetVertexPosition(VertexID);
 			Delta /= EditorToolProperties->MorphTargetWeight;
