[MorphTargetEditingTool] Improve posing perf

diff --git a/Engine/Plugins/Experimental/Animation/SkeletalMeshMorphTargetEditingTools/Source/SkeletalMeshMorphTargetEditingTools/Private/MorphTargetVertexSculptTool.cpp b/Engine/Plugins/Experimental/Animation/SkeletalMeshMorphTargetEditingTools/Source/SkeletalMeshMorphTargetEditingTools/Private/MorphTargetVertexSculptTool.cpp
--- a/Engine/Plugins/Experimental/Animation/SkeletalMeshMorphTargetEditingTools/Source/SkeletalMeshMorphTargetEditingTools/Private/MorphTargetVertexSculptTool.cpp
+++ b/Engine/Plugins/Experimental/Animation/SkeletalMeshMorphTargetEditingTools/Source/SkeletalMeshMorphTargetEditingTools/Private/MorphTargetVertexSculptTool.cpp
@@ -59,6 +59,23 @@ void UMorphTargetVertexSculptTool::Setup()
 	SymmetryProperties->bSymmetryCanBeEnabled = false;
 
 	OnToolMeshChangedDelegate = DynamicMeshComponent->OnMeshRegionChanged.AddUObject(this, &UMorphTargetVertexSculptTool::OnToolMeshChanged);
+
+	UPrimitiveComponent* Component = Cast<IPrimitiveComponentBackedTarget>(GetTarget())->GetOwnerComponent();
+	USkeletalMeshComponent* SkeletalMeshComponent = Cast<USkeletalMeshComponent>(Component);
+	
+	ToolMeshDescription = *SkeletalMeshComponent->GetSkeletalMeshAsset()->GetMeshDescription(LodIndex0);
+
+	// Pre compact the mesh description to avoid compacting one every frame
+	if (ToolMeshDescription.NeedsCompact())
+	{
+		FElementIDRemappings Remappings;
+		ToolMeshDescription.Compact(Remappings);
+	}
+	else
+	{
+		// Make sure indexers are built before entering parallel work
+		ToolMeshDescription.BuildVertexIndexers();
+	}
 	
 	InitializeCache();
 }
@@ -205,27 +222,24 @@ void UMorphTargetVertexSculptTool::UpdateCacheIfNeeded()
 
 	TMap<FName, float> MorphTargetWeights = PreviousMorphWeights;
 	MorphTargetWeights.Remove(CurrentMorphTarget);
-	
-	FMeshDescription SculptMeshDescription;
-	FStaticMeshAttributes Attributes(SculptMeshDescription);
-	Attributes.Register();
-	
-	FMeshDescription NonDeformedMesh;
 
-	FDynamicMeshToMeshDescription Converter;
-	Converter.Convert(GetSculptMesh(), SculptMeshDescription, false);
+	// In the case that source mesh contains non-manifold verts, extra duplicated verts will be added to end of the dynamic mesh verts array
+	// see FMeshDescriptionToDynamicMesh::Convert
+	check(GetSculptMesh()->VertexCount() >= ToolMeshDescription.GetVertexPositions().GetNumElements());
+	for (FVertexID VertexID : ToolMeshDescription.Vertices().GetElementIDs())
+	{
+		ToolMeshDescription.GetVertexPositions()[VertexID] = FVector3f(GetSculptMesh()->GetVertex(VertexID));
+	}
 	
 	const FMeshDescription* RefMeshDescription = SkeletalMeshComponent->GetSkeletalMeshAsset()->GetMeshDescription(LodIndex0);
 
 	const TArray<FTransform>& ComponentSpaceTransforms = PreviousPoseComponentSpace;
 	
-	if (FSkeletalMeshOperations::GetMorphTargetDeltasFromPosedMesh(SculptMeshDescription, *RefMeshDescription, RefComponentSpaceTransforms, NonDeformedMesh, ComponentSpaceTransforms, NAME_None, MorphTargetWeights))
+	if (FSkeletalMeshOperations::GetUnposedMeshInPlace(ToolMeshDescription, *RefMeshDescription, RefComponentSpaceTransforms, ComponentSpaceTransforms, NAME_None, MorphTargetWeights))
 	{
-		// Worth noting that the sculpt mesh description can have more verts due to the initial construction of the dynamic mesh by the tool target when
-		// the mesh has non-manifold structure, which is why we want to iterate over the RefMeshDescription
-		for (FVertexID VertexID : RefMeshDescription->Vertices().GetElementIDs())
+		for (FVertexID VertexID : ToolMeshDescription.Vertices().GetElementIDs())
 		{
-			FVector3f Delta = NonDeformedMesh.GetVertexPosition(VertexID) - RefMeshDescription->GetVertexPosition(VertexID);
+			FVector3f Delta = ToolMeshDescription.GetVertexPosition(VertexID) - RefMeshDescription->GetVertexPosition(VertexID);
 			Delta /= EditorToolProperties->MorphTargetWeight;
 			
 			Results.Add(Delta);
@@ -235,8 +249,6 @@ void UMorphTargetVertexSculptTool::UpdateCacheIfNeeded()
 	CachedDelta = MoveTemp(Results);
 
 	bCached = true;
-
-	GetToolManager()->DisplayMessage(LOCTEXT("UpdateCacheNotification", "Cache Updated"), EToolMessageLevel::UserNotification);
 }
 
 void UMorphTargetVertexSculptTool::PoseToolMesh()
@@ -368,11 +380,8 @@ void UMorphTargetVertexSculptTool::PoseToolMesh()
 		RefComponentSpaceTransforms.Add(FAnimationRuntime::GetComponentSpaceTransformRefPose(SkeletalMeshComponent->GetSkeletalMeshAsset()->GetRefSkeleton(), BoneIndex));
 	}
 			
-	FMeshDescription PosedMeshDescription;
-
 	{
-		FMeshDescription MeshDescriptionWithAllMorphs = *SkeletalMeshComponent->GetSkeletalMeshAsset()->GetMeshDescription(LodIndex0);
-		FSkeletalMeshAttributes Attributes(MeshDescriptionWithAllMorphs);
+		FSkeletalMeshAttributes Attributes(ToolMeshDescription);
 
 		if (!Attributes.GetMorphTargetNames().Contains(CurrentMorphTarget))
 		{
@@ -381,17 +390,24 @@ void UMorphTargetVertexSculptTool::PoseToolMesh()
 
 		TVertexAttributesRef<FVector3f> MorphTargetPosDeltaAttribute = Attributes.GetVertexMorphPositionDelta(CurrentMorphTarget);
 
-		for (FVertexID VertexID : MeshDescriptionWithAllMorphs.Vertices().GetElementIDs())
+		for (FVertexID VertexID : ToolMeshDescription.Vertices().GetElementIDs())
 		{
 			MorphTargetPosDeltaAttribute[VertexID] = CachedDelta[VertexID];
 		}
 
-		FSkeletalMeshOperations::GetPosedMesh(MeshDescriptionWithAllMorphs, PosedMeshDescription, ComponentSpaceTransforms, NAME_None, MorphTargetWeights);
+		FMeshDescription* RefMeshDescription = SkeletalMeshComponent->GetSkeletalMeshAsset()->GetMeshDescription(LodIndex0);
+		// Reset verts to ref pose for posing
+		for (FVertexID VertexID : ToolMeshDescription.Vertices().GetElementIDs())
+		{
+			ToolMeshDescription.GetVertexPositions()[VertexID] = RefMeshDescription->GetVertexPositions()[VertexID];
+		}
+
+		FSkeletalMeshOperations::GetPosedMeshInPlace(ToolMeshDescription, ComponentSpaceTransforms, NAME_None, MorphTargetWeights);
 	}
 
 	FDynamicMesh3& Mesh = *GetSculptMesh();
 	
-	TArrayView<const FVector3f> VertexPositionsAttribute = PosedMeshDescription.GetVertexPositions().GetRawArray();
+	TArrayView<const FVector3f> VertexPositionsAttribute = ToolMeshDescription.GetVertexPositions().GetRawArray();
 	for (int32 Index = 0; Index < VertexPositionsAttribute.Num(); ++Index)
 	{
 		Mesh.SetVertex(Index, FVector3d(VertexPositionsAttribute[Index]));
diff --git a/Engine/Plugins/Experimental/Animation/SkeletalMeshMorphTargetEditingTools/Source/SkeletalMeshMorphTargetEditingTools/Public/MorphTargetVertexSculptTool.h b/Engine/Plugins/Experimental/Animation/SkeletalMeshMorphTargetEditingTools/Source/SkeletalMeshMorphTargetEditingTools/Public/MorphTargetVertexSculptTool.h
--- a/Engine/Plugins/Experimental/Animation/SkeletalMeshMorphTargetEditingTools/Source/SkeletalMeshMorphTargetEditingTools/Public/MorphTargetVertexSculptTool.h
+++ b/Engine/Plugins/Experimental/Animation/SkeletalMeshMorphTargetEditingTools/Source/SkeletalMeshMorphTargetEditingTools/Public/MorphTargetVertexSculptTool.h
@@ -8,6 +8,7 @@
 #include "BaseTools/MultiSelectionMeshEditingTool.h"
 #include "MorphTargetEditingToolProperties.h"
 #include "SkeletalMeshMorphTargetBackedTarget.h"
+#include "MeshDescription.h"
 
 #include "MorphTargetVertexSculptTool.generated.h"
 
@@ -73,7 +74,9 @@ public:
 
 	FDelegateHandle OnToolMeshChangedDelegate;
 
-	static constexpr int32 LodIndex0 = 0;	
+	static constexpr int32 LodIndex0 = 0;
+
+	FMeshDescription ToolMeshDescription;
 };
 
 
