[Morph Target Editing Tools] Further improve posing perf

diff --git a/Engine/Plugins/Experimental/Animation/SkeletalMeshMorphTargetEditingTools/Source/SkeletalMeshMorphTargetEditingTools/Private/MorphTargetVertexSculptTool.cpp b/Engine/Plugins/Experimental/Animation/SkeletalMeshMorphTargetEditingTools/Source/SkeletalMeshMorphTargetEditingTools/Private/MorphTargetVertexSculptTool.cpp
--- a/Engine/Plugins/Experimental/Animation/SkeletalMeshMorphTargetEditingTools/Source/SkeletalMeshMorphTargetEditingTools/Private/MorphTargetVertexSculptTool.cpp
+++ b/Engine/Plugins/Experimental/Animation/SkeletalMeshMorphTargetEditingTools/Source/SkeletalMeshMorphTargetEditingTools/Private/MorphTargetVertexSculptTool.cpp
@@ -60,22 +60,7 @@ void UMorphTargetVertexSculptTool::Setup()
 
 	OnToolMeshChangedDelegate = DynamicMeshComponent->OnMeshRegionChanged.AddUObject(this, &UMorphTargetVertexSculptTool::OnToolMeshChanged);
 
-	UPrimitiveComponent* Component = Cast<IPrimitiveComponentBackedTarget>(GetTarget())->GetOwnerComponent();
-	USkeletalMeshComponent* SkeletalMeshComponent = Cast<USkeletalMeshComponent>(Component);
-	
-	ToolMeshDescription = *SkeletalMeshComponent->GetSkeletalMeshAsset()->GetMeshDescription(LodIndex0);
 
-	// Pre compact the mesh description to avoid compacting one every frame
-	if (ToolMeshDescription.NeedsCompact())
-	{
-		FElementIDRemappings Remappings;
-		ToolMeshDescription.Compact(Remappings);
-	}
-	else
-	{
-		// Make sure indexers are built before entering parallel work
-		ToolMeshDescription.BuildVertexIndexers();
-	}
 	
 	InitializeCache();
 }
@@ -90,7 +75,7 @@ void UMorphTargetVertexSculptTool::Shutdown(EToolShutdownType ShutdownType)
 	UpdateCacheIfNeeded();
 	
 	MorphTargetBackedTarget->SetEditingMorphTargetName(EditorToolProperties->GetEditingMorphTargetName());
-	MorphTargetBackedTarget->SetDataToCommit(&CachedDelta);
+	MorphTargetBackedTarget->SetDataToCommit(MoveTemp(ToolMeshDescription), ToolMorphTargetName);
 	
 	Super::Shutdown(ShutdownType);
 
@@ -167,28 +152,32 @@ void UMorphTargetVertexSculptTool::InitializeCache()
 {
 	UPrimitiveComponent* Component = Cast<IPrimitiveComponentBackedTarget>(GetTarget())->GetOwnerComponent();
 	USkeletalMeshComponent* SkeletalMeshComponent = Cast<USkeletalMeshComponent>(Component);
+
+	ToolMeshDescription = *SkeletalMeshComponent->GetSkeletalMeshAsset()->GetMeshDescription(LodIndex0);
+
+	// Pre compact the mesh description to avoid compacting one every frame
+	if (ToolMeshDescription.NeedsCompact())
+	{
+		FElementIDRemappings Remappings;
+		ToolMeshDescription.Compact(Remappings);
+	}
+	else
+	{
+		// Make sure indexers are built before entering parallel work
+		ToolMeshDescription.BuildVertexIndexers();
+	}
 	
-	FMeshDescription* MeshDescription = SkeletalMeshComponent->GetSkeletalMeshAsset()->GetMeshDescription(LodIndex0);
-	int32 NumVerts = MeshDescription->Vertices().Num();
+	FSkeletalMeshAttributes Attributes(ToolMeshDescription);
 	
-	CachedDelta.Init(FVector3f::Zero(), NumVerts);
-
-	if (EditorToolProperties->Operation == EMorphTargetEditorOperation::Edit)
+	ToolMorphTargetName = EditorToolProperties->GetEditingMorphTargetName();
+	
+	if (EditorToolProperties->Operation == EMorphTargetEditorOperation::New)
 	{
-		FSkeletalMeshAttributes Attributes(*MeshDescription);
-		FName CurrentMorphTarget = EditorToolProperties->GetEditingMorphTargetName();
-		check(Attributes.GetMorphTargetNames().Contains(CurrentMorphTarget));
-
-		TVertexAttributesRef<FVector3f> MorphTargetPosDeltaAttribute = Attributes.GetVertexMorphPositionDelta(CurrentMorphTarget);
-
-		for (FVertexID VertexID : MeshDescription->Vertices().GetElementIDs())
-		{
-			CachedDelta[VertexID] = MorphTargetPosDeltaAttribute[VertexID];
-		}	
+		check(!Attributes.GetMorphTargetNames().Contains(ToolMorphTargetName))
+		Attributes.RegisterMorphTargetAttribute(ToolMorphTargetName, false);
 	}
 
 	bCached = true;
-	
 }
 
 void UMorphTargetVertexSculptTool::UpdateCacheIfNeeded()
@@ -205,8 +194,6 @@ void UMorphTargetVertexSculptTool::UpdateCacheIfNeeded()
 		return;
 	}
 	
-	TArray<FVector3f> Results;
-	
 	UPrimitiveComponent* Component = Cast<IPrimitiveComponentBackedTarget>(GetTarget())->GetOwnerComponent();
 	USkeletalMeshComponent* SkeletalMeshComponent = Cast<USkeletalMeshComponent>(Component);
 	const FReferenceSkeleton& RefSkeleton = SkeletalMeshComponent->GetSkeletalMeshAsset()->GetRefSkeleton();
@@ -218,10 +205,9 @@ void UMorphTargetVertexSculptTool::UpdateCacheIfNeeded()
 		RefComponentSpaceTransforms.Add(FAnimationRuntime::GetComponentSpaceTransformRefPose(RefSkeleton, BoneIndex));
 	}
 
-	FName CurrentMorphTarget = EditorToolProperties->GetEditingMorphTargetName();
 
 	TMap<FName, float> MorphTargetWeights = PreviousMorphWeights;
-	MorphTargetWeights.Remove(CurrentMorphTarget);
+	MorphTargetWeights.Remove(ToolMorphTargetName);
 
 	// In the case that source mesh contains non-manifold verts, extra duplicated verts will be added to end of the dynamic mesh verts array
 	// see FMeshDescriptionToDynamicMesh::Convert
@@ -231,9 +217,12 @@ void UMorphTargetVertexSculptTool::UpdateCacheIfNeeded()
 		ToolMeshDescription.GetVertexPositions()[VertexID] = FVector3f(GetSculptMesh()->GetVertex(VertexID));
 	}
 	
-	const FMeshDescription* RefMeshDescription = SkeletalMeshComponent->GetSkeletalMeshAsset()->GetMeshDescription(LodIndex0);
+	FMeshDescription* RefMeshDescription = SkeletalMeshComponent->GetSkeletalMeshAsset()->GetMeshDescription(LodIndex0);
 
 	const TArray<FTransform>& ComponentSpaceTransforms = PreviousPoseComponentSpace;
+
+	FSkeletalMeshAttributes Attributes(ToolMeshDescription);
+	TVertexAttributesRef<FVector3f> MorphTargetPosDeltaAttribute = Attributes.GetVertexMorphPositionDelta(ToolMorphTargetName);
 	
 	if (FSkeletalMeshOperations::GetUnposedMeshInPlace(ToolMeshDescription, *RefMeshDescription, RefComponentSpaceTransforms, ComponentSpaceTransforms, NAME_None, MorphTargetWeights))
 	{
@@ -242,12 +231,10 @@ void UMorphTargetVertexSculptTool::UpdateCacheIfNeeded()
 			FVector3f Delta = ToolMeshDescription.GetVertexPosition(VertexID) - RefMeshDescription->GetVertexPosition(VertexID);
 			Delta /= EditorToolProperties->MorphTargetWeight;
 			
-			Results.Add(Delta);
+			MorphTargetPosDeltaAttribute[VertexID] = Delta;
 		}
 	}
 
-	CachedDelta = MoveTemp(Results);
-
 	bCached = true;
 }
 
@@ -278,8 +265,6 @@ void UMorphTargetVertexSculptTool::PoseToolMesh()
 	UPrimitiveComponent* Component = Cast<IPrimitiveComponentBackedTarget>(GetTarget())->GetOwnerComponent();
 	USkeletalMeshComponent* SkeletalMeshComponent = Cast<USkeletalMeshComponent>(Component);
 	
-	FName CurrentMorphTarget = EditorToolProperties->GetEditingMorphTargetName();
-		
 	TMap<FName, float> MorphTargetWeights;
 
 	for (const TPair<const UMorphTarget*, int32>& MorphTarget: SkeletalMeshComponent->ActiveMorphTargets)
@@ -290,7 +275,7 @@ void UMorphTargetVertexSculptTool::PoseToolMesh()
 		MorphTargetWeights.Add(MorphName, MorphWeight);
 	}
 	
-	MorphTargetWeights.FindOrAdd(CurrentMorphTarget) = EditorToolProperties->MorphTargetWeight;
+	MorphTargetWeights.FindOrAdd(ToolMorphTargetName) = EditorToolProperties->MorphTargetWeight;
 
 	const TArray<FTransform>& ComponentSpaceTransforms = SkeletalMeshComponent->GetComponentSpaceTransforms();
 	
@@ -380,31 +365,17 @@ void UMorphTargetVertexSculptTool::PoseToolMesh()
 		RefComponentSpaceTransforms.Add(FAnimationRuntime::GetComponentSpaceTransformRefPose(SkeletalMeshComponent->GetSkeletalMeshAsset()->GetRefSkeleton(), BoneIndex));
 	}
 			
+	FMeshDescription* RefMeshDescription = SkeletalMeshComponent->GetSkeletalMeshAsset()->GetMeshDescription(LodIndex0);
+	// Reset verts to ref pose for posing
+	for (FVertexID VertexID : ToolMeshDescription.Vertices().GetElementIDs())
 	{
-		FSkeletalMeshAttributes Attributes(ToolMeshDescription);
-
-		if (!Attributes.GetMorphTargetNames().Contains(CurrentMorphTarget))
-		{
-			Attributes.RegisterMorphTargetAttribute(CurrentMorphTarget, false);
-		}
-
-		TVertexAttributesRef<FVector3f> MorphTargetPosDeltaAttribute = Attributes.GetVertexMorphPositionDelta(CurrentMorphTarget);
-
-		for (FVertexID VertexID : ToolMeshDescription.Vertices().GetElementIDs())
-		{
-			MorphTargetPosDeltaAttribute[VertexID] = CachedDelta[VertexID];
-		}
-
-		FMeshDescription* RefMeshDescription = SkeletalMeshComponent->GetSkeletalMeshAsset()->GetMeshDescription(LodIndex0);
-		// Reset verts to ref pose for posing
-		for (FVertexID VertexID : ToolMeshDescription.Vertices().GetElementIDs())
-		{
-			ToolMeshDescription.GetVertexPositions()[VertexID] = RefMeshDescription->GetVertexPositions()[VertexID];
-		}
-
-		FSkeletalMeshOperations::GetPosedMeshInPlace(ToolMeshDescription, ComponentSpaceTransforms, NAME_None, MorphTargetWeights);
+		ToolMeshDescription.GetVertexPositions()[VertexID] = RefMeshDescription->GetVertexPositions()[VertexID];
 	}
 
+	// No need to compute normals for the mesh description, the sculpt dynamic mesh recomputes its own normals during OnDynamicMeshComponentChanged()
+	constexpr bool bSkipRecomputeNormalsTangents = true;
+	FSkeletalMeshOperations::GetPosedMeshInPlace(ToolMeshDescription, ComponentSpaceTransforms, NAME_None, MorphTargetWeights, bSkipRecomputeNormalsTangents);
+
 	FDynamicMesh3& Mesh = *GetSculptMesh();
 	
 	TArrayView<const FVector3f> VertexPositionsAttribute = ToolMeshDescription.GetVertexPositions().GetRawArray();
diff --git a/Engine/Plugins/Experimental/Animation/SkeletalMeshMorphTargetEditingTools/Source/SkeletalMeshMorphTargetEditingTools/Public/MorphTargetVertexSculptTool.h b/Engine/Plugins/Experimental/Animation/SkeletalMeshMorphTargetEditingTools/Source/SkeletalMeshMorphTargetEditingTools/Public/MorphTargetVertexSculptTool.h
--- a/Engine/Plugins/Experimental/Animation/SkeletalMeshMorphTargetEditingTools/Source/SkeletalMeshMorphTargetEditingTools/Public/MorphTargetVertexSculptTool.h
+++ b/Engine/Plugins/Experimental/Animation/SkeletalMeshMorphTargetEditingTools/Source/SkeletalMeshMorphTargetEditingTools/Public/MorphTargetVertexSculptTool.h
@@ -63,7 +63,6 @@ public:
 	TWeakInterfacePtr<ISkeletalMeshMorphTargetBackedTarget> MorphTargetBackedTarget; 
 	
 	bool bCached = false;
-	TArray<FVector3f> CachedDelta;
 
 	bool bPoseOpInitialized = false;
 	TArray<FTransform> PreviousPoseComponentSpace;
@@ -77,6 +76,7 @@ public:
 	static constexpr int32 LodIndex0 = 0;
 
 	FMeshDescription ToolMeshDescription;
+	FName ToolMorphTargetName;
 };
 
 
diff --git a/Engine/Source/Runtime/SkeletalMeshDescription/Private/SkeletalMeshOperations.cpp b/Engine/Source/Runtime/SkeletalMeshDescription/Private/SkeletalMeshOperations.cpp
--- a/Engine/Source/Runtime/SkeletalMeshDescription/Private/SkeletalMeshOperations.cpp
+++ b/Engine/Source/Runtime/SkeletalMeshDescription/Private/SkeletalMeshOperations.cpp
@@ -275,7 +275,8 @@ static void PoseMesh(
 	FMeshDescription& InOutTargetMesh,
 	TConstArrayView<FMatrix44f> InRefToUserTransforms, 
 	const FName InSkinWeightProfile, 
-	const TMap<FName, float>& InMorphTargetWeights
+	const TMap<FName, float>& InMorphTargetWeights,
+	bool bInSkipRecomputeNormalsTangents = false
 	)
 {
 	struct FMorphInfo
@@ -340,7 +341,7 @@ static void PoseMesh(
 		};
 		TArray<FMorphProcessContext> Contexts;
 		ParallelForWithTaskContext(Contexts, InOutTargetMesh.Vertices().Num(),
-			[&Mesh=InOutTargetMesh, &PositionAttribute, &MorphInfos, bAllMorphNormalsValid](FMorphProcessContext& Context, int32 Index)
+			[&Mesh=InOutTargetMesh, &PositionAttribute, &MorphInfos, bAllMorphNormalsValid, bInSkipRecomputeNormalsTangents](FMorphProcessContext& Context, int32 Index)
 			{
 				const FVertexID VertexID{Index};
 
@@ -362,7 +363,7 @@ static void PoseMesh(
 				{
 					PositionAttribute.Set(VertexID, Position);
 					
-					if (!bAllMorphNormalsValid)
+					if (!bAllMorphNormalsValid && !bInSkipRecomputeNormalsTangents)
 					{
 						Context.DirtyVertexInstances.Append(Mesh.GetVertexVertexInstanceIDs(VertexID));
 
@@ -378,7 +379,7 @@ static void PoseMesh(
 		if (bAllMorphNormalsValid)
 		{
 			ParallelForWithTaskContext(Contexts, InOutTargetMesh.VertexInstances().Num(),
-				[&MorphInfos, &NormalAttribute, &TangentAttribute, &BinormalSignsAttribute](FMorphProcessContext& Context, int32 Index)
+				[&MorphInfos, &NormalAttribute, &TangentAttribute, &BinormalSignsAttribute, bInSkipRecomputeNormalsTangents](FMorphProcessContext& Context, int32 Index)
 				{
 					FVertexInstanceID VertexInstanceID{Index};
 
@@ -411,37 +412,43 @@ static void PoseMesh(
 						}
 						else
 						{
-							// Something went wrong. Reconstruct the normal from the tangent and binormal.
-							Context.DirtyVertexInstances.Add(VertexInstanceID);
+							if (!bInSkipRecomputeNormalsTangents)
+							{
+								// Something went wrong. Reconstruct the normal from the tangent and binormal.
+								Context.DirtyVertexInstances.Add(VertexInstanceID);
+							}
 						}
 					}
 				});
 		}
-		
-		// Clear out any normals that were affected by the point move, or ended up being degenerate during normal offsetting.
-		TSet<FVertexInstanceID> DirtyVertexInstances;
-		for (const FMorphProcessContext& ProcessContext: Contexts)
-		{
-			DirtyVertexInstances.Append(ProcessContext.DirtyVertexInstances);
-		}
-		
-		if (!DirtyVertexInstances.IsEmpty())
+
+		if (!bInSkipRecomputeNormalsTangents)
 		{
-			// Mark any vector as zero that we want to regenerate from triangle + neighbors + tangents.
-			for (const FVertexInstanceID VertexInstanceID: DirtyVertexInstances)
+			// Clear out any normals that were affected by the point move, or ended up being degenerate during normal offsetting.
+			TSet<FVertexInstanceID> DirtyVertexInstances;
+			for (const FMorphProcessContext& ProcessContext: Contexts)
 			{
-				NormalAttribute.Set(VertexInstanceID, FVector3f::ZeroVector);
+				DirtyVertexInstances.Append(ProcessContext.DirtyVertexInstances);
 			}
+		
+			if (!DirtyVertexInstances.IsEmpty())
+			{
+				// Mark any vector as zero that we want to regenerate from triangle + neighbors + tangents.
+				for (const FVertexInstanceID VertexInstanceID: DirtyVertexInstances)
+				{
+					NormalAttribute.Set(VertexInstanceID, FVector3f::ZeroVector);
+				}
 
-			FSkeletalMeshOperations::ComputeTriangleTangentsAndNormals(InOutTargetMesh, UE_SMALL_NUMBER, nullptr);
+				FSkeletalMeshOperations::ComputeTriangleTangentsAndNormals(InOutTargetMesh, UE_SMALL_NUMBER, nullptr);
 
-			// Compute the normals on the dirty vertices, and adjust the tangents to match.
-			FSkeletalMeshOperations::ComputeTangentsAndNormals(InOutTargetMesh, EComputeNTBsFlags::WeightedNTBs);
+				// Compute the normals on the dirty vertices, and adjust the tangents to match.
+				FSkeletalMeshOperations::ComputeTangentsAndNormals(InOutTargetMesh, EComputeNTBsFlags::WeightedNTBs);
 			
-			// We don't need the triangle tangents and normals anymore.
-			InOutTargetMesh.TriangleAttributes().UnregisterAttribute(MeshAttribute::Triangle::Normal);
-			InOutTargetMesh.TriangleAttributes().UnregisterAttribute(MeshAttribute::Triangle::Tangent);
-			InOutTargetMesh.TriangleAttributes().UnregisterAttribute(MeshAttribute::Triangle::Binormal);
+				// We don't need the triangle tangents and normals anymore.
+				InOutTargetMesh.TriangleAttributes().UnregisterAttribute(MeshAttribute::Triangle::Normal);
+				InOutTargetMesh.TriangleAttributes().UnregisterAttribute(MeshAttribute::Triangle::Tangent);
+				InOutTargetMesh.TriangleAttributes().UnregisterAttribute(MeshAttribute::Triangle::Binormal);
+			}	
 		}
 	}
 
@@ -654,6 +661,75 @@ static void UnposeMesh(
 					}
 				}
 			});
+
+		if (bAllMorphNormalsValid)
+		{
+			ParallelForWithTaskContext(Contexts, InOutTargetMesh.VertexInstances().Num(),
+				[&MorphInfos, &NormalAttribute, &TangentAttribute, &BinormalSignsAttribute](FMorphProcessContext& Context, int32 Index)
+				{
+					FVertexInstanceID VertexInstanceID{Index};
+
+					FVector3f Normal = NormalAttribute.Get(VertexInstanceID);
+					FVector3f Tangent = TangentAttribute.Get(VertexInstanceID);
+					FVector3f Binormal = FVector3f::CrossProduct(Normal, Tangent) * BinormalSignsAttribute.Get(VertexInstanceID);
+					
+					bool bMoved = false;
+					for (const FMorphInfo& MorphInfo: MorphInfos)
+					{
+						const FVector3f NormalDelta = MorphInfo.NormalDelta.Get(VertexInstanceID) * MorphInfo.Weight * -1.0f;
+						if (!NormalDelta.IsNearlyZero())
+						{
+							Normal += NormalDelta;
+							bMoved = true;
+						}
+					}
+
+					if (bMoved)
+					{
+						if (Normal.Normalize())
+						{
+							// Badly named function. This orthonormalizes X & Y using Z as the control. 
+							FVector3f::CreateOrthonormalBasis(Tangent, Binormal, Normal);
+							
+							NormalAttribute.Set(VertexInstanceID, Normal);
+							TangentAttribute.Set(VertexInstanceID, Tangent);
+							const float BinormalSign = FMatrix44f(Tangent, Binormal, Normal, FVector3f::ZeroVector).Determinant() < 0 ? -1.0f : +1.0f;
+							BinormalSignsAttribute.Set(VertexInstanceID, BinormalSign);
+						}
+						else
+						{
+							// Something went wrong. Reconstruct the normal from the tangent and binormal.
+							Context.DirtyVertexInstances.Add(VertexInstanceID);
+						}
+					}
+				});
+		}
+
+		// Clear out any normals that were affected by the point move, or ended up being degenerate during normal offsetting.
+		TSet<FVertexInstanceID> DirtyVertexInstances;
+		for (const FMorphProcessContext& ProcessContext: Contexts)
+		{
+			DirtyVertexInstances.Append(ProcessContext.DirtyVertexInstances);
+		}
+		
+		if (!DirtyVertexInstances.IsEmpty())
+		{
+			// Mark any vector as zero that we want to regenerate from triangle + neighbors + tangents.
+			for (const FVertexInstanceID VertexInstanceID: DirtyVertexInstances)
+			{
+				NormalAttribute.Set(VertexInstanceID, FVector3f::ZeroVector);
+			}
+
+			FSkeletalMeshOperations::ComputeTriangleTangentsAndNormals(InOutTargetMesh, UE_SMALL_NUMBER, nullptr);
+
+			// Compute the normals on the dirty vertices, and adjust the tangents to match.
+			FSkeletalMeshOperations::ComputeTangentsAndNormals(InOutTargetMesh, EComputeNTBsFlags::WeightedNTBs);
+			
+			// We don't need the triangle tangents and normals anymore.
+			InOutTargetMesh.TriangleAttributes().UnregisterAttribute(MeshAttribute::Triangle::Normal);
+			InOutTargetMesh.TriangleAttributes().UnregisterAttribute(MeshAttribute::Triangle::Tangent);
+			InOutTargetMesh.TriangleAttributes().UnregisterAttribute(MeshAttribute::Triangle::Binormal);
+		}	
 	}
 }
 	
@@ -732,7 +808,7 @@ bool FSkeletalMeshOperations::GetPosedMesh(
 }
 
 bool FSkeletalMeshOperations::GetPosedMeshInPlace(FMeshDescription& InOutTargetMesh, TConstArrayView<FTransform> InComponentSpaceTransforms,
-	const FName InSkinWeightProfile, const TMap<FName, float>& InMorphTargetWeights, bool bInWriteBonePose)
+	const FName InSkinWeightProfile, const TMap<FName, float>& InMorphTargetWeights, bool bInSkipRecomputeNormalsTangents, bool bInWriteBonePose)
 {
 	// Verify that the mesh is valid.
 	FSkeletalMeshConstAttributes Attributes(InOutTargetMesh);
@@ -777,7 +853,7 @@ bool FSkeletalMeshOperations::GetPosedMeshInPlace(FMeshDescription& InOutTargetM
 	}
 	
 	// Pose the target mesh in-place.
-	UE::Impl::PoseMesh(InOutTargetMesh, RefToUserTransforms, InSkinWeightProfile, InMorphTargetWeights);
+	UE::Impl::PoseMesh(InOutTargetMesh, RefToUserTransforms, InSkinWeightProfile, InMorphTargetWeights, bInSkipRecomputeNormalsTangents);
 
 	if (bInWriteBonePose)
 	{
