[Morph Target Editing Tool] Added missing files from last submission

diff --git a/Engine/Plugins/Experimental/Animation/SkeletalMeshMorphTargetEditingTools/Source/SkeletalMeshMorphTargetEditingTools/Private/SkeletalMeshMorphTargetToolTarget.cpp b/Engine/Plugins/Experimental/Animation/SkeletalMeshMorphTargetEditingTools/Source/SkeletalMeshMorphTargetEditingTools/Private/SkeletalMeshMorphTargetToolTarget.cpp
--- a/Engine/Plugins/Experimental/Animation/SkeletalMeshMorphTargetEditingTools/Source/SkeletalMeshMorphTargetEditingTools/Private/SkeletalMeshMorphTargetToolTarget.cpp
+++ b/Engine/Plugins/Experimental/Animation/SkeletalMeshMorphTargetEditingTools/Source/SkeletalMeshMorphTargetEditingTools/Private/SkeletalMeshMorphTargetToolTarget.cpp
@@ -32,46 +32,49 @@ UE::Geometry::FDynamicMesh3 USkeletalMeshMorphTargetToolTarget::GetDynamicMesh(c
 
 void USkeletalMeshMorphTargetToolTarget::CommitDynamicMesh(const UE::Geometry::FDynamicMesh3& Mesh, const FDynamicMeshCommitInfo& CommitInfo)
 {
-	if (!DataPtr->IsEmpty())
+	FMeshDescription* TargetMeshDescription = GetEditingMeshDescription();
+	FSkeletalMeshAttributes TargetAttributes(*TargetMeshDescription);
+
+	if (!TargetAttributes.GetMorphTargetNames().Contains(EditingMorphTargetName))
 	{
-		FMeshDescription* TargetMeshDescription = GetEditingMeshDescription();
-		FSkeletalMeshAttributes TargetAttributes(*TargetMeshDescription);
-
-		if (!TargetAttributes.GetMorphTargetNames().Contains(EditingMorphTargetName))
-		{
-			TargetAttributes.RegisterMorphTargetAttribute(EditingMorphTargetName, false);
-		}
-		
-		TVertexAttributesRef<FVector3f> TargetMorphTargetPosDeltaAttribute = TargetAttributes.GetVertexMorphPositionDelta(EditingMorphTargetName);
-
-		for (FVertexID VertexID : TargetMeshDescription->Vertices().GetElementIDs())
-		{
-			TargetMorphTargetPosDeltaAttribute[VertexID] = (*DataPtr)[VertexID.GetValue()];
-		}		
+		TargetAttributes.RegisterMorphTargetAttribute(EditingMorphTargetName, false);
+	}
+	
+	TVertexAttributesRef<FVector3f> TargetMorphTargetPosDeltaAttribute = TargetAttributes.GetVertexMorphPositionDelta(EditingMorphTargetName);
+
 	
-		CommitEditedMeshDescription();
-
-		USkeletalMesh* SkeletalMesh = GetSkeletalMesh();
-		UAnimCurveMetaData* AnimCurveMetaData = SkeletalMesh->GetAssetUserData<UAnimCurveMetaData>();
-		if (AnimCurveMetaData == nullptr)
-		{
-			AnimCurveMetaData = NewObject<UAnimCurveMetaData>(SkeletalMesh, NAME_None, RF_Transactional);
-			SkeletalMesh->AddAssetUserData(AnimCurveMetaData);
-		}
-
-		AnimCurveMetaData->AddCurveMetaData(EditingMorphTargetName);
-
-		// Ensure we have a morph flag set
-		if (FCurveMetaData* CurveMetaData = AnimCurveMetaData->GetCurveMetaData(EditingMorphTargetName))
-		{
-			CurveMetaData->Type.bMorphtarget = true;
-		}
+	FSkeletalMeshAttributes ToolAttributes(ToolMeshDescription);
+	TVertexAttributesRef<FVector3f> ToolMorphTargetPosDeltaAttribute = ToolAttributes.GetVertexMorphPositionDelta(ToolMorphTargetName);
+	
+	
+	for (FVertexID VertexID : TargetMeshDescription->Vertices().GetElementIDs())
+	{
+		TargetMorphTargetPosDeltaAttribute[VertexID] = ToolMorphTargetPosDeltaAttribute[VertexID.GetValue()];
+	}		
+
+	CommitEditedMeshDescription();
+
+	USkeletalMesh* SkeletalMesh = GetSkeletalMesh();
+	UAnimCurveMetaData* AnimCurveMetaData = SkeletalMesh->GetAssetUserData<UAnimCurveMetaData>();
+	if (AnimCurveMetaData == nullptr)
+	{
+		AnimCurveMetaData = NewObject<UAnimCurveMetaData>(SkeletalMesh, NAME_None, RF_Transactional);
+		SkeletalMesh->AddAssetUserData(AnimCurveMetaData);
+	}
+
+	AnimCurveMetaData->AddCurveMetaData(EditingMorphTargetName);
+
+	// Ensure we have a morph flag set
+	if (FCurveMetaData* CurveMetaData = AnimCurveMetaData->GetCurveMetaData(EditingMorphTargetName))
+	{
+		CurveMetaData->Type.bMorphtarget = true;
 	}
 }
 
-void USkeletalMeshMorphTargetToolTarget::SetDataToCommit(TArray<FVector3f>* InDataPtr)
+void USkeletalMeshMorphTargetToolTarget::SetDataToCommit(FMeshDescription&& InMeshDescription, const FName& InToolMorphTargetName)
 {
-	DataPtr = InDataPtr;
+	ToolMeshDescription = MoveTemp(InMeshDescription);
+	ToolMorphTargetName = InToolMorphTargetName;
 }
 
 void USkeletalMeshMorphTargetToolTarget::SetEditingMorphTargetName(const FName& InName)
diff --git a/Engine/Plugins/Experimental/Animation/SkeletalMeshMorphTargetEditingTools/Source/SkeletalMeshMorphTargetEditingTools/Public/SkeletalMeshMorphTargetBackedTarget.h b/Engine/Plugins/Experimental/Animation/SkeletalMeshMorphTargetEditingTools/Source/SkeletalMeshMorphTargetEditingTools/Public/SkeletalMeshMorphTargetBackedTarget.h
--- a/Engine/Plugins/Experimental/Animation/SkeletalMeshMorphTargetEditingTools/Source/SkeletalMeshMorphTargetEditingTools/Public/SkeletalMeshMorphTargetBackedTarget.h
+++ b/Engine/Plugins/Experimental/Animation/SkeletalMeshMorphTargetEditingTools/Source/SkeletalMeshMorphTargetEditingTools/Public/SkeletalMeshMorphTargetBackedTarget.h
@@ -7,6 +7,8 @@
 
 #include "SkeletalMeshMorphTargetBackedTarget.generated.h"
 
+struct FMeshDescription;
+
 UINTERFACE(MinimalAPI)
 class USkeletalMeshMorphTargetBackedTarget : public UInterface
 {
@@ -17,7 +19,7 @@ class ISkeletalMeshMorphTargetBackedTarget : public IInterface
 {
 	GENERATED_BODY()
 public:
-	virtual void SetDataToCommit(TArray<FVector3f>* InDataPtr) = 0;
+	virtual void SetDataToCommit(FMeshDescription&& InMeshDescription, const FName& ToolMorphTargetName) = 0;
 	virtual void SetEditingMorphTargetName(const FName& InName) = 0;
 	virtual TArray<FName> GetEditableMorphTargetNames() = 0;
 	virtual FName GetValidNameForNewMorphTarget(const FName& InName) = 0;
diff --git a/Engine/Plugins/Experimental/Animation/SkeletalMeshMorphTargetEditingTools/Source/SkeletalMeshMorphTargetEditingTools/Public/SkeletalMeshMorphTargetToolTarget.h b/Engine/Plugins/Experimental/Animation/SkeletalMeshMorphTargetEditingTools/Source/SkeletalMeshMorphTargetEditingTools/Public/SkeletalMeshMorphTargetToolTarget.h
--- a/Engine/Plugins/Experimental/Animation/SkeletalMeshMorphTargetEditingTools/Source/SkeletalMeshMorphTargetEditingTools/Public/SkeletalMeshMorphTargetToolTarget.h
+++ b/Engine/Plugins/Experimental/Animation/SkeletalMeshMorphTargetEditingTools/Source/SkeletalMeshMorphTargetEditingTools/Public/SkeletalMeshMorphTargetToolTarget.h
@@ -7,6 +7,7 @@
 #include "ToolTargets/SkeletalMeshComponentToolTarget.h"
 #include "SkeletalMeshMorphTargetBackedTarget.h"
 #include "ToolTargets/ToolTarget.h"
+#include "MeshDescription.h"
 
 #include "SkeletalMeshMorphTargetToolTarget.generated.h"
 
@@ -31,7 +32,7 @@ public:
 	using IDynamicMeshCommitter::CommitDynamicMesh; // unhide the other overload
 
 
-	virtual void SetDataToCommit(TArray<FVector3f>* InDataPtr) override;
+	virtual void SetDataToCommit(FMeshDescription&& InMeshDescription, const FName& InToolMorphTargetName) override;
 	virtual void SetEditingMorphTargetName(const FName& InName) override;
 	virtual TArray<FName> GetEditableMorphTargetNames() override;
 	virtual FName GetValidNameForNewMorphTarget(const FName& InName) override;
@@ -47,7 +48,8 @@ protected:
 
 	static constexpr int32 LodIndex0 = 0;
 
-	TArray<FVector3f>* DataPtr = nullptr;
+	FMeshDescription ToolMeshDescription;
+	FName ToolMorphTargetName;
 };
 
 
diff --git a/Engine/Source/Runtime/SkeletalMeshDescription/Public/SkeletalMeshOperations.h b/Engine/Source/Runtime/SkeletalMeshDescription/Public/SkeletalMeshOperations.h
--- a/Engine/Source/Runtime/SkeletalMeshDescription/Public/SkeletalMeshOperations.h
+++ b/Engine/Source/Runtime/SkeletalMeshDescription/Public/SkeletalMeshOperations.h
@@ -113,6 +113,7 @@ public:
 	 *  @param InComponentSpaceTransforms The component space transforms used to move the mesh joints for deforming, using linear-blend skinning.
 	 *  @param InSkinWeightProfile The skin weight profile to use as the source of skin weights for the deformation.
 	 *  @param InMorphTargetWeights Optional morph target weights to apply. Any morph target that doesn't exist is ignored.
+	 *  @param bInSkipRecomputeNormalsTangents skip recompute normals and tangents to improve performance
 	 *  @param bInWriteBonePose whether to write the current bone pose into the mesh description.
 	 *  @return \c true if the operation succeeded.
 	 */
@@ -121,6 +122,7 @@ public:
 		TConstArrayView<FTransform> InComponentSpaceTransforms,
 		const FName InSkinWeightProfile = NAME_None,
 		const TMap<FName, float>& InMorphTargetWeights = {},
+		bool bInSkipRecomputeNormalsTangents = false,
 		bool bInWriteBonePose = false
 		);
 
