[Morph editing tools] Fixed issues related to redo + fixed ensure when undo

diff --git a/Engine/Plugins/Experimental/Animation/SkeletalMeshMorphTargetEditingTools/Source/SkeletalMeshMorphTargetEditingTools/Private/MorphTargetVertexSculptTool.cpp b/Engine/Plugins/Experimental/Animation/SkeletalMeshMorphTargetEditingTools/Source/SkeletalMeshMorphTargetEditingTools/Private/MorphTargetVertexSculptTool.cpp
--- a/Engine/Plugins/Experimental/Animation/SkeletalMeshMorphTargetEditingTools/Source/SkeletalMeshMorphTargetEditingTools/Private/MorphTargetVertexSculptTool.cpp
+++ b/Engine/Plugins/Experimental/Animation/SkeletalMeshMorphTargetEditingTools/Source/SkeletalMeshMorphTargetEditingTools/Private/MorphTargetVertexSculptTool.cpp
@@ -124,9 +124,15 @@ void UMorphTargetVertexSculptTool::OnPropertyModified(UObject* PropertySet, FPro
 	}
 }
 
+int32 UMorphTargetVertexSculptTool::FindHitSculptMeshTriangle(const FRay3d& LocalRay)
+{
+	WaitForPendingUndoRedo();
+	return Super::FindHitSculptMeshTriangle(LocalRay);
+}
+
 void UMorphTargetVertexSculptTool::OnToolMeshChanged(UDynamicMeshComponent* Component, const FMeshRegionChangeBase* Change, bool bRevert)
 {
-	if (bRevert)
+	if (!bPosingSculptMesh)
 	{
 		bCached = false;
 		UpdateCacheIfNeeded();
@@ -334,7 +340,8 @@ void UMorphTargetVertexSculptTool::PoseToolMesh()
 		if (bPoseChangedLastTick)
 		{
 			FMeshReplacementChange Change(MeshBeforePosing, MakeShared<FDynamicMesh3>(*GetSculptMesh()));
-			
+
+			TGuardValue<bool> PosingSculptMeshScope(bPosingSculptMesh, true);
 			OnDynamicMeshComponentChanged(DynamicMeshComponent, &Change, false);
 		}
 
diff --git a/Engine/Plugins/Experimental/Animation/SkeletalMeshMorphTargetEditingTools/Source/SkeletalMeshMorphTargetEditingTools/Public/MorphTargetVertexSculptTool.h b/Engine/Plugins/Experimental/Animation/SkeletalMeshMorphTargetEditingTools/Source/SkeletalMeshMorphTargetEditingTools/Public/MorphTargetVertexSculptTool.h
--- a/Engine/Plugins/Experimental/Animation/SkeletalMeshMorphTargetEditingTools/Source/SkeletalMeshMorphTargetEditingTools/Public/MorphTargetVertexSculptTool.h
+++ b/Engine/Plugins/Experimental/Animation/SkeletalMeshMorphTargetEditingTools/Source/SkeletalMeshMorphTargetEditingTools/Public/MorphTargetVertexSculptTool.h
@@ -43,6 +43,9 @@ public:
 	virtual void Shutdown(EToolShutdownType ShutdownType) override;
 	virtual void OnTick(float DeltaTime) override;
 	virtual void OnPropertyModified(UObject* PropertySet, FProperty* Property) override;
+protected:
+	virtual int32 FindHitSculptMeshTriangle(const FRay3d& LocalRay) override;
+	
 
 	void OnToolMeshChanged(UDynamicMeshComponent* Component, const FMeshRegionChangeBase* Change, bool bRevert);
 
@@ -77,6 +80,8 @@ public:
 
 	FMeshDescription ToolMeshDescription;
 	FName ToolMorphTargetName;
+
+	bool bPosingSculptMesh = false;
 };
 
 
