[MorphTargetEditingTools] Fixed fast redo lead to exploded mesh

diff --git a/Engine/Plugins/Experimental/Animation/SkeletalMeshMorphTargetEditingTools/Source/SkeletalMeshMorphTargetEditingTools/Private/MorphTargetVertexSculptTool.cpp b/Engine/Plugins/Experimental/Animation/SkeletalMeshMorphTargetEditingTools/Source/SkeletalMeshMorphTargetEditingTools/Private/MorphTargetVertexSculptTool.cpp
--- a/Engine/Plugins/Experimental/Animation/SkeletalMeshMorphTargetEditingTools/Source/SkeletalMeshMorphTargetEditingTools/Private/MorphTargetVertexSculptTool.cpp
+++ b/Engine/Plugins/Experimental/Animation/SkeletalMeshMorphTargetEditingTools/Source/SkeletalMeshMorphTargetEditingTools/Private/MorphTargetVertexSculptTool.cpp
@@ -124,6 +124,33 @@ void UMorphTargetVertexSculptTool::OnPropertyModified(UObject* PropertySet, FPro
 	}
 }
 
+void UMorphTargetVertexSculptTool::OnEndStroke()
+{
+	// update spatial
+	bTargetDirty = true;
+
+	GetActiveBrushOp()->EndStroke(GetSculptMesh(), LastStamp, VertexROI);
+
+	check(ActiveVertexChange);
+
+	TUniquePtr<TWrappedToolCommandChange<FMeshVertexChange>> NewChange = MakeUnique<TWrappedToolCommandChange<FMeshVertexChange>>();
+	NewChange->WrappedChange = MoveTemp(ActiveVertexChange->Change);
+	NewChange->BeforeModify = [this](bool bRevert)
+	{
+		this->WaitForPendingUndoRedo();
+		// Any sculpt change needs to be applied to the latest posed mesh during undo / redo, so pose the mesh immediately instead of waiting
+		// for the next tick
+		this->PoseToolMesh();
+	};
+
+	GetToolManager()->EmitObjectChange(DynamicMeshComponent, MoveTemp(NewChange), LOCTEXT("VertexSculptChange", "Brush Stroke"));
+
+	LongTransactions.Close(GetToolManager());
+
+	delete ActiveVertexChange;
+	ActiveVertexChange = nullptr;
+}
+
 int32 UMorphTargetVertexSculptTool::FindHitSculptMeshTriangle(const FRay3d& LocalRay)
 {
 	WaitForPendingUndoRedo();
@@ -134,8 +161,8 @@ void UMorphTargetVertexSculptTool::OnToolMeshChanged(UDynamicMeshComponent* Comp
 {
 	if (!bPosingSculptMesh)
 	{
+		// If not posing, it must have been a sculpt change that was reverted/applied, should update the cache
 		bCached = false;
-		UpdateCacheIfNeeded();
 	}
 }
 
@@ -188,6 +215,9 @@ void UMorphTargetVertexSculptTool::InitializeCache()
 
 void UMorphTargetVertexSculptTool::UpdateCacheIfNeeded()
 {
+	WaitForPendingStampUpdate();
+	WaitForPendingUndoRedo();
+	
 	if (bCached)
 	{
 		return;
diff --git a/Engine/Plugins/Experimental/Animation/SkeletalMeshMorphTargetEditingTools/Source/SkeletalMeshMorphTargetEditingTools/Public/MorphTargetVertexSculptTool.h b/Engine/Plugins/Experimental/Animation/SkeletalMeshMorphTargetEditingTools/Source/SkeletalMeshMorphTargetEditingTools/Public/MorphTargetVertexSculptTool.h
--- a/Engine/Plugins/Experimental/Animation/SkeletalMeshMorphTargetEditingTools/Source/SkeletalMeshMorphTargetEditingTools/Public/MorphTargetVertexSculptTool.h
+++ b/Engine/Plugins/Experimental/Animation/SkeletalMeshMorphTargetEditingTools/Source/SkeletalMeshMorphTargetEditingTools/Public/MorphTargetVertexSculptTool.h
@@ -43,6 +43,9 @@ public:
 	virtual void Shutdown(EToolShutdownType ShutdownType) override;
 	virtual void OnTick(float DeltaTime) override;
 	virtual void OnPropertyModified(UObject* PropertySet, FProperty* Property) override;
+	virtual void OnEndStroke() override;
+
+	
 protected:
 	virtual int32 FindHitSculptMeshTriangle(const FRay3d& LocalRay) override;
 	
