[MorphTargetEditing]  partial support for symmetry while sculpting

diff --git a/Engine/Plugins/Experimental/Animation/SkeletalMeshMorphTargetEditingTools/Source/SkeletalMeshMorphTargetEditingTools/Private/MorphTargetVertexSculptTool.cpp b/Engine/Plugins/Experimental/Animation/SkeletalMeshMorphTargetEditingTools/Source/SkeletalMeshMorphTargetEditingTools/Private/MorphTargetVertexSculptTool.cpp
--- a/Engine/Plugins/Experimental/Animation/SkeletalMeshMorphTargetEditingTools/Source/SkeletalMeshMorphTargetEditingTools/Private/MorphTargetVertexSculptTool.cpp
+++ b/Engine/Plugins/Experimental/Animation/SkeletalMeshMorphTargetEditingTools/Source/SkeletalMeshMorphTargetEditingTools/Private/MorphTargetVertexSculptTool.cpp
@@ -5,6 +5,7 @@
 #include "ContextObjectStore.h"
 #include "ToolBuilderUtil.h"
 #include "InteractiveToolManager.h"
+#include "InteractiveToolChange.h"
 #include "ModelingToolTargetUtil.h"
 #include "ToolTargetManager.h"
 #include "StaticMeshAttributes.h"
@@ -56,8 +57,6 @@ void UMorphTargetVertexSculptTool::Setup()
 
 	MorphTargetBackedTarget = Cast<ISkeletalMeshMorphTargetBackedTarget>(GetTarget());
 
-	SymmetryProperties->bSymmetryCanBeEnabled = false;
-
 	OnToolMeshChangedDelegate = DynamicMeshComponent->OnMeshRegionChanged.AddUObject(this, &UMorphTargetVertexSculptTool::OnToolMeshChanged);
 
 
@@ -124,6 +123,33 @@ void UMorphTargetVertexSculptTool::OnPropertyModified(UObject* PropertySet, FPro
 	}
 }
 
+/*
+ * internal Change classes
+ */
+
+class FMorphTargetVertexSculptNonSymmetricChange : public FToolCommandChange
+{
+public:
+	virtual void Apply(UObject* Object) override;
+	virtual void Revert(UObject* Object) override;
+};
+
+
+void FMorphTargetVertexSculptNonSymmetricChange::Apply(UObject* Object)
+{
+	if (Cast<UMorphTargetVertexSculptTool>(Object))
+	{
+		Cast<UMorphTargetVertexSculptTool>(Object)->UndoRedo_RestoreSymmetryPossibleState(false);
+	}
+}
+void FMorphTargetVertexSculptNonSymmetricChange::Revert(UObject* Object)
+{
+	if (Cast<UMorphTargetVertexSculptTool>(Object))
+	{
+		Cast<UMorphTargetVertexSculptTool>(Object)->UndoRedo_RestoreSymmetryPossibleState(true);
+	}
+}
+
 void UMorphTargetVertexSculptTool::OnEndStroke()
 {
 	// update spatial
@@ -144,7 +170,13 @@ void UMorphTargetVertexSculptTool::OnEndStroke()
 	};
 
 	GetToolManager()->EmitObjectChange(DynamicMeshComponent, MoveTemp(NewChange), LOCTEXT("VertexSculptChange", "Brush Stroke"));
-
+	if (bMeshSymmetryIsValid && bApplySymmetry == false)
+	{
+		// if we end a stroke while symmetry is possible but disabled, we now have to assume that symmetry is no longer possible
+		GetToolManager()->EmitObjectChange(this, MakeUnique<FMorphTargetVertexSculptNonSymmetricChange>(), LOCTEXT("DisableSymmetryChange", "Disable Symmetry"));
+		bMeshSymmetryIsValid = false;
+		SymmetryProperties->bSymmetryCanBeEnabled = bMeshSymmetryIsValid;
+	}
 	LongTransactions.Close(GetToolManager());
 
 	delete ActiveVertexChange;
diff --git a/Engine/Plugins/Experimental/Animation/SkeletalMeshMorphTargetEditingTools/Source/SkeletalMeshMorphTargetEditingTools/Public/MorphTargetVertexSculptTool.h b/Engine/Plugins/Experimental/Animation/SkeletalMeshMorphTargetEditingTools/Source/SkeletalMeshMorphTargetEditingTools/Public/MorphTargetVertexSculptTool.h
--- a/Engine/Plugins/Experimental/Animation/SkeletalMeshMorphTargetEditingTools/Source/SkeletalMeshMorphTargetEditingTools/Public/MorphTargetVertexSculptTool.h
+++ b/Engine/Plugins/Experimental/Animation/SkeletalMeshMorphTargetEditingTools/Source/SkeletalMeshMorphTargetEditingTools/Public/MorphTargetVertexSculptTool.h
@@ -47,6 +47,7 @@ public:
 
 	
 protected:
+	friend class FMorphTargetVertexSculptNonSymmetricChange;
 	virtual int32 FindHitSculptMeshTriangle(const FRay3d& LocalRay) override;
 	
 
