[MorphTargetEditingTool] avoid adding empty morphs into the mesh descriptions

diff --git a/Engine/Plugins/Experimental/Animation/SkeletalMeshMorphTargetEditingTools/Source/SkeletalMeshMorphTargetEditingTools/Private/MorphTargetVertexSculptTool.cpp b/Engine/Plugins/Experimental/Animation/SkeletalMeshMorphTargetEditingTools/Source/SkeletalMeshMorphTargetEditingTools/Private/MorphTargetVertexSculptTool.cpp
--- a/Engine/Plugins/Experimental/Animation/SkeletalMeshMorphTargetEditingTools/Source/SkeletalMeshMorphTargetEditingTools/Private/MorphTargetVertexSculptTool.cpp
+++ b/Engine/Plugins/Experimental/Animation/SkeletalMeshMorphTargetEditingTools/Source/SkeletalMeshMorphTargetEditingTools/Private/MorphTargetVertexSculptTool.cpp
@@ -72,9 +72,16 @@ void UMorphTargetVertexSculptTool::Shutdown(EToolShutdownType ShutdownType)
 	}
 	
 	UpdateCacheIfNeeded();
-	
-	MorphTargetBackedTarget->SetEditingMorphTargetName(EditorToolProperties->GetEditingMorphTargetName());
-	MorphTargetBackedTarget->SetDataToCommit(MoveTemp(ToolMeshDescription), ToolMorphTargetName);
+
+	if (bHasValidData)
+	{
+		MorphTargetBackedTarget->SetEditingMorphTargetName(EditorToolProperties->GetEditingMorphTargetName());
+		MorphTargetBackedTarget->SetDataToCommit(MoveTemp(ToolMeshDescription), ToolMorphTargetName);
+	}
+	else
+	{
+		GetToolManager()->DisplayMessage(LOCTEXT("MorphTargetInvalidData", "Morph Target is empty, no change was made") , EToolMessageLevel::UserNotification);
+	}
 	
 	Super::Shutdown(ShutdownType);
 
@@ -294,11 +301,16 @@ void UMorphTargetVertexSculptTool::UpdateCacheIfNeeded()
 	
 	if (FSkeletalMeshOperations::GetUnposedMeshInPlace(ToolMeshDescription, *RefMeshDescription, RefComponentSpaceTransforms, ComponentSpaceTransforms, NAME_None, MorphTargetWeights))
 	{
+		bHasValidData = false;
 		for (FVertexID VertexID : ToolMeshDescription.Vertices().GetElementIDs())
 		{
 			FVector3f Delta = ToolMeshDescription.GetVertexPosition(VertexID) - RefMeshDescription->GetVertexPosition(VertexID);
 			Delta /= EditorToolProperties->MorphTargetWeight;
-			
+
+			if (!bHasValidData && Delta.SizeSquared() > FMath::Square(UE_THRESH_POINTS_ARE_NEAR))
+			{
+				bHasValidData = true;
+			}
 			MorphTargetPosDeltaAttribute[VertexID] = Delta;
 		}
 	}
diff --git a/Engine/Plugins/Experimental/Animation/SkeletalMeshMorphTargetEditingTools/Source/SkeletalMeshMorphTargetEditingTools/Private/SkeletalMeshMorphTargetToolTarget.cpp b/Engine/Plugins/Experimental/Animation/SkeletalMeshMorphTargetEditingTools/Source/SkeletalMeshMorphTargetEditingTools/Private/SkeletalMeshMorphTargetToolTarget.cpp
--- a/Engine/Plugins/Experimental/Animation/SkeletalMeshMorphTargetEditingTools/Source/SkeletalMeshMorphTargetEditingTools/Private/SkeletalMeshMorphTargetToolTarget.cpp
+++ b/Engine/Plugins/Experimental/Animation/SkeletalMeshMorphTargetEditingTools/Source/SkeletalMeshMorphTargetEditingTools/Private/SkeletalMeshMorphTargetToolTarget.cpp
@@ -32,6 +32,11 @@ UE::Geometry::FDynamicMesh3 USkeletalMeshMorphTargetToolTarget::GetDynamicMesh(c
 
 void USkeletalMeshMorphTargetToolTarget::CommitDynamicMesh(const UE::Geometry::FDynamicMesh3& Mesh, const FDynamicMeshCommitInfo& CommitInfo)
 {
+	if (ToolMeshDescription.IsEmpty() || ToolMorphTargetName.IsNone() || EditingMorphTargetName.IsNone())
+	{
+		return;
+	}
+	
 	FMeshDescription* TargetMeshDescription = GetEditingMeshDescription();
 	FSkeletalMeshAttributes TargetAttributes(*TargetMeshDescription);
 
diff --git a/Engine/Plugins/Experimental/Animation/SkeletalMeshMorphTargetEditingTools/Source/SkeletalMeshMorphTargetEditingTools/Public/MorphTargetVertexSculptTool.h b/Engine/Plugins/Experimental/Animation/SkeletalMeshMorphTargetEditingTools/Source/SkeletalMeshMorphTargetEditingTools/Public/MorphTargetVertexSculptTool.h
--- a/Engine/Plugins/Experimental/Animation/SkeletalMeshMorphTargetEditingTools/Source/SkeletalMeshMorphTargetEditingTools/Public/MorphTargetVertexSculptTool.h
+++ b/Engine/Plugins/Experimental/Animation/SkeletalMeshMorphTargetEditingTools/Source/SkeletalMeshMorphTargetEditingTools/Public/MorphTargetVertexSculptTool.h
@@ -70,6 +70,7 @@ protected:
 	TWeakInterfacePtr<ISkeletalMeshMorphTargetBackedTarget> MorphTargetBackedTarget; 
 	
 	bool bCached = false;
+	bool bHasValidData = false;
 
 	bool bPoseOpInitialized = false;
 	TArray<FTransform> PreviousPoseComponentSpace;
