[MorphTargetEditingTools] Fixed crash when closing the engine without exiting the tool

diff --git a/Engine/Plugins/Experimental/Animation/SkeletalMeshMorphTargetEditingTools/Source/SkeletalMeshMorphTargetEditingTools/Private/MorphTargetVertexSculptTool.cpp b/Engine/Plugins/Experimental/Animation/SkeletalMeshMorphTargetEditingTools/Source/SkeletalMeshMorphTargetEditingTools/Private/MorphTargetVertexSculptTool.cpp
--- a/Engine/Plugins/Experimental/Animation/SkeletalMeshMorphTargetEditingTools/Source/SkeletalMeshMorphTargetEditingTools/Private/MorphTargetVertexSculptTool.cpp
+++ b/Engine/Plugins/Experimental/Animation/SkeletalMeshMorphTargetEditingTools/Source/SkeletalMeshMorphTargetEditingTools/Private/MorphTargetVertexSculptTool.cpp
@@ -70,17 +70,21 @@ void UMorphTargetVertexSculptTool::Shutdown(EToolShutdownType ShutdownType)
 	{
 		DynamicMeshComponent->OnMeshChanged.Remove(OnDynamicMeshComponentChangedHandle);
 	}
-	
-	UpdateCacheIfNeeded();
 
-	if (bHasValidData)
-	{
-		MorphTargetBackedTarget->SetEditingMorphTargetName(EditorToolProperties->GetEditingMorphTargetName());
-		MorphTargetBackedTarget->SetDataToCommit(MoveTemp(ToolMeshDescription), ToolMorphTargetName);
-	}
-	else
+	if (ShutdownType == EToolShutdownType::Accept)
 	{
-		GetToolManager()->DisplayMessage(LOCTEXT("MorphTargetInvalidData", "Morph Target is empty, no change was made") , EToolMessageLevel::UserNotification);
+		UpdateCacheIfNeeded();
+    
+    	if (bHasValidData)
+    	{
+    		MorphTargetBackedTarget->SetEditingMorphTargetName(EditorToolProperties->GetEditingMorphTargetName());
+    		MorphTargetBackedTarget->SetDataToCommit(MoveTemp(ToolMeshDescription), ToolMorphTargetName);
+    	}
+    	else
+    	{
+    		GetToolManager()->DisplayMessage(LOCTEXT("MorphTargetInvalidData", "Morph Target is empty, no change was made") , EToolMessageLevel::UserNotification);
+    		ShutdownType = EToolShutdownType::Cancel;
+    	}	
 	}
 	
 	Super::Shutdown(ShutdownType);
@@ -256,6 +260,13 @@ void UMorphTargetVertexSculptTool::UpdateCacheIfNeeded()
 {
 	WaitForPendingStampUpdate();
 	WaitForPendingUndoRedo();
+	UPrimitiveComponent* Component = Cast<IPrimitiveComponentBackedTarget>(GetTarget())->GetOwnerComponent();
+	USkeletalMeshComponent* SkeletalMeshComponent = Cast<USkeletalMeshComponent>(Component);
+
+	if (!SkeletalMeshComponent->GetSkeletalMeshAsset())
+	{
+		return;
+	}
 	
 	if (bCached)
 	{
@@ -268,9 +279,7 @@ void UMorphTargetVertexSculptTool::UpdateCacheIfNeeded()
 		bCached = true;
 		return;
 	}
-	
-	UPrimitiveComponent* Component = Cast<IPrimitiveComponentBackedTarget>(GetTarget())->GetOwnerComponent();
-	USkeletalMeshComponent* SkeletalMeshComponent = Cast<USkeletalMeshComponent>(Component);
+
 	const FReferenceSkeleton& RefSkeleton = SkeletalMeshComponent->GetSkeletalMeshAsset()->GetRefSkeleton();
 	
 	TArray<FTransform> RefComponentSpaceTransforms;
