[Morph Editing Tool] Added a simple erase tool

diff --git a/Engine/Plugins/Experimental/Animation/SkeletalMeshMorphTargetEditingTools/Source/SkeletalMeshMorphTargetEditingTools/Private/EraseMorphTargetBrushOps.h b/Engine/Plugins/Experimental/Animation/SkeletalMeshMorphTargetEditingTools/Source/SkeletalMeshMorphTargetEditingTools/Private/EraseMorphTargetBrushOps.h
new file mode 100644
--- /dev/null
+++ b/Engine/Plugins/Experimental/Animation/SkeletalMeshMorphTargetEditingTools/Source/SkeletalMeshMorphTargetEditingTools/Private/EraseMorphTargetBrushOps.h
@@ -0,0 +1,78 @@
+// Copyright Epic Games, Inc. All Rights Reserved.
+
+#pragma once
+
+#include "Sculpting/MeshBrushOpBase.h"
+#include "DynamicMesh/DynamicMesh3.h"
+#include "Async/ParallelFor.h"
+#include "EraseMorphTargetBrushOps.generated.h"
+
+
+UCLASS(MinimalAPI)
+class UEraseMorphTargetBrushOpProps : public UMeshSculptBrushOpProps
+{
+	GENERATED_BODY()
+public:
+	/** Strength of the Brush */
+	UPROPERTY(EditAnywhere, Category = InflateBrush, meta = (DisplayName = "Strength", UIMin = "0.0", UIMax = "1.0", ClampMin = "0.0", ClampMax = "1.0", ModelingQuickEdit, ModelingQuickSettings = 200))
+	float Strength = 0.5;
+
+	/** Amount of falloff to apply */
+	UPROPERTY(EditAnywhere, Category = InflateBrush, meta = (DisplayName = "Falloff", UIMin = "0.0", UIMax = "1.0", ClampMin = "0.0", ClampMax = "1.0"))
+	float Falloff = 0.5;
+
+	virtual float GetStrength() override { return Strength; }
+	virtual void SetStrength(float NewStrength) override { Strength = FMathf::Clamp(NewStrength, 0.0f, 1.0f); }
+	virtual float GetFalloff() override { return Falloff; }
+	virtual bool SupportsStrengthPressure() override { return true; }
+};
+
+class FEraseMorphTargetBrushOp : public FMeshSculptBrushOp
+{
+
+public:
+	double BrushSpeedTuning = 2.0;
+
+	typedef TUniqueFunction<const FDynamicMesh3*()> GetMeshWithoutCurrentMorphFuncType;
+
+	GetMeshWithoutCurrentMorphFuncType GetMeshWithoutCurrentMorphFunc;
+	
+	FEraseMorphTargetBrushOp(GetMeshWithoutCurrentMorphFuncType InFunc)
+	{
+		GetMeshWithoutCurrentMorphFunc = MoveTemp(InFunc);
+	}
+	
+	virtual ESculptBrushOpTargetType GetBrushTargetType() const override
+	{
+		return ESculptBrushOpTargetType::SculptMesh;
+	}
+
+	virtual void ApplyStamp(const FDynamicMesh3* Mesh, const FSculptBrushStamp& Stamp, const TArray<int32>& Vertices, TArray<FVector3d>& NewPositionsOut) override
+	{
+		const FDynamicMesh3* MeshWithoutCurentMorph = GetMeshWithoutCurrentMorphFunc();
+		
+		const FVector3d& StampPos = Stamp.LocalFrame.Origin;
+
+		double UsePower = Stamp.Power * Stamp.Radius * Stamp.DeltaTime * BrushSpeedTuning;
+
+		ParallelFor(Vertices.Num(), [&](int32 k)
+		{
+			int32 VertIdx = Vertices[k];
+			FVector3d OrigPos = Mesh->GetVertex(VertIdx);
+			FVector3d TargetPos = MeshWithoutCurentMorph->GetVertex(VertIdx);
+			FVector3d MaxDelta = TargetPos - OrigPos; 
+
+			FVector3d MoveVec = UE::Geometry::Normalized(TargetPos - OrigPos, FMathd::ZeroTolerance) * UsePower;
+
+			double Falloff = GetFalloff().Evaluate(Stamp, OrigPos);
+
+			MoveVec = MoveVec * Falloff;
+			MoveVec = MoveVec.GetClampedToMaxSize(MaxDelta.Length());
+			
+			FVector3d NewPos = OrigPos + MoveVec;
+			NewPositionsOut[k] = NewPos;
+		});
+	}
+
+};
+
diff --git a/Engine/Plugins/Experimental/Animation/SkeletalMeshMorphTargetEditingTools/Source/SkeletalMeshMorphTargetEditingTools/Private/MorphTargetVertexSculptTool.cpp b/Engine/Plugins/Experimental/Animation/SkeletalMeshMorphTargetEditingTools/Source/SkeletalMeshMorphTargetEditingTools/Private/MorphTargetVertexSculptTool.cpp
--- a/Engine/Plugins/Experimental/Animation/SkeletalMeshMorphTargetEditingTools/Source/SkeletalMeshMorphTargetEditingTools/Private/MorphTargetVertexSculptTool.cpp
+++ b/Engine/Plugins/Experimental/Animation/SkeletalMeshMorphTargetEditingTools/Source/SkeletalMeshMorphTargetEditingTools/Private/MorphTargetVertexSculptTool.cpp
@@ -22,6 +22,7 @@
 #include "TargetInterfaces/DynamicMeshProvider.h"
 #include "TargetInterfaces/PrimitiveComponentBackedTarget.h"
 #include "AnimationRuntime.h"
+#include "EraseMorphTargetBrushOps.h"
 #include "SkeletalMeshAttributes.h"
 
 #define LOCTEXT_NAMESPACE "MorphTargetVertexSculptTool"
@@ -59,11 +60,25 @@ void UMorphTargetVertexSculptTool::Setup()
 
 	OnToolMeshChangedDelegate = DynamicMeshComponent->OnMeshRegionChanged.AddUObject(this, &UMorphTargetVertexSculptTool::OnToolMeshChanged);
 
-
+	MeshWithoutCurrentMorph = *GetSculptMesh();
 	
 	InitializeCache();
 }
 
+void UMorphTargetVertexSculptTool::RegisterBrushes()
+{
+	Super::RegisterBrushes();
+	
+	GetMeshWithoutCurrentMorphFunc = [this](){return &MeshWithoutCurrentMorph;};
+
+	// Had to hijack the EraseSculptLayer identifier from base mesh vertex sculpt tool for our erase morph target tool since it is the simplest way to get an icon for the tool.
+	RegisterBrushType(
+		(int32)EMeshVertexSculptBrushType::EraseSculptLayer,
+		LOCTEXT("EraseSculptLayerBrushTypeName", "EraseSculptLayer"),
+		MakeUnique<FLambdaMeshSculptBrushOpFactory>([this]() { return MakeUnique<FEraseMorphTargetBrushOp>(GetMeshWithoutCurrentMorphFunc);}),
+		NewObject<UEraseMorphTargetBrushOpProps>(this));
+}
+
 void UMorphTargetVertexSculptTool::Shutdown(EToolShutdownType ShutdownType)
 {
 	if (DynamicMeshComponent != nullptr)
@@ -464,9 +479,16 @@ void UMorphTargetVertexSculptTool::PoseToolMesh()
 	FDynamicMesh3& Mesh = *GetSculptMesh();
 	
 	TArrayView<const FVector3f> VertexPositionsAttribute = ToolMeshDescription.GetVertexPositions().GetRawArray();
+	
+	FSkeletalMeshAttributes Attributes(ToolMeshDescription);
+	TVertexAttributesRef<FVector3f> MorphTargetPosDeltaAttribute = Attributes.GetVertexMorphPositionDelta(ToolMorphTargetName);
+	
 	for (int32 Index = 0; Index < VertexPositionsAttribute.Num(); ++Index)
 	{
 		Mesh.SetVertex(Index, FVector3d(VertexPositionsAttribute[Index]));
+
+		FVector3f Delta = MorphTargetPosDeltaAttribute[Index];
+		MeshWithoutCurrentMorph.SetVertex(Index, FVector3d(VertexPositionsAttribute[Index] - Delta));
 	}	
 	
 	DynamicMeshComponent->FastNotifyPositionsUpdated();
diff --git a/Engine/Plugins/Experimental/Animation/SkeletalMeshMorphTargetEditingTools/Source/SkeletalMeshMorphTargetEditingTools/Public/MorphTargetVertexSculptTool.h b/Engine/Plugins/Experimental/Animation/SkeletalMeshMorphTargetEditingTools/Source/SkeletalMeshMorphTargetEditingTools/Public/MorphTargetVertexSculptTool.h
--- a/Engine/Plugins/Experimental/Animation/SkeletalMeshMorphTargetEditingTools/Source/SkeletalMeshMorphTargetEditingTools/Public/MorphTargetVertexSculptTool.h
+++ b/Engine/Plugins/Experimental/Animation/SkeletalMeshMorphTargetEditingTools/Source/SkeletalMeshMorphTargetEditingTools/Public/MorphTargetVertexSculptTool.h
@@ -48,6 +48,7 @@ public:
 	
 protected:
 	friend class FMorphTargetVertexSculptNonSymmetricChange;
+	virtual void RegisterBrushes() override
 
 	void OnToolMeshChanged(UDynamicMeshComponent* Component, const FMeshRegionChangeBase* Change, bool bRevert);
 
@@ -66,6 +67,9 @@ protected:
 	TObjectPtr<class UMorphTargetEditingToolProperties> EditorToolProperties;
 
 	TWeakInterfacePtr<ISkeletalMeshMorphTargetBackedTarget> MorphTargetBackedTarget; 
+
+	TFunction<const FDynamicMesh3*()> GetMeshWithoutCurrentMorphFunc;
+	FDynamicMesh3 MeshWithoutCurrentMorph;
 	
 	bool bCached = false;
 	bool bHasValidData = false;
