[Morph Editing Tools] Fixed erase tool erasing too much if the active weight is not 1.0

diff --git a/Engine/Plugins/Experimental/Animation/SkeletalMeshMorphTargetEditingTools/Source/SkeletalMeshMorphTargetEditingTools/Private/MorphTargetVertexSculptTool.cpp b/Engine/Plugins/Experimental/Animation/SkeletalMeshMorphTargetEditingTools/Source/SkeletalMeshMorphTargetEditingTools/Private/MorphTargetVertexSculptTool.cpp
--- a/Engine/Plugins/Experimental/Animation/SkeletalMeshMorphTargetEditingTools/Source/SkeletalMeshMorphTargetEditingTools/Private/MorphTargetVertexSculptTool.cpp
+++ b/Engine/Plugins/Experimental/Animation/SkeletalMeshMorphTargetEditingTools/Source/SkeletalMeshMorphTargetEditingTools/Private/MorphTargetVertexSculptTool.cpp
@@ -487,7 +487,7 @@ void UMorphTargetVertexSculptTool::PoseToolMesh()
 	{
 		Mesh.SetVertex(Index, FVector3d(VertexPositionsAttribute[Index]));
 
-		FVector3f Delta = MorphTargetPosDeltaAttribute[Index];
+		FVector3f Delta = MorphTargetPosDeltaAttribute[Index] * EditorToolProperties->MorphTargetWeight;
 		MeshWithoutCurrentMorph.SetVertex(Index, FVector3d(VertexPositionsAttribute[Index] - Delta));
 	}	
 	
