[Morph Editing Tool] Fixed typo

diff --git a/Engine/Plugins/Experimental/Animation/SkeletalMeshMorphTargetEditingTools/Source/SkeletalMeshMorphTargetEditingTools/Public/MorphTargetVertexSculptTool.h b/Engine/Plugins/Experimental/Animation/SkeletalMeshMorphTargetEditingTools/Source/SkeletalMeshMorphTargetEditingTools/Public/MorphTargetVertexSculptTool.h
--- a/Engine/Plugins/Experimental/Animation/SkeletalMeshMorphTargetEditingTools/Source/SkeletalMeshMorphTargetEditingTools/Public/MorphTargetVertexSculptTool.h
+++ b/Engine/Plugins/Experimental/Animation/SkeletalMeshMorphTargetEditingTools/Source/SkeletalMeshMorphTargetEditingTools/Public/MorphTargetVertexSculptTool.h
@@ -48,7 +48,7 @@ public:
 	
 protected:
 	friend class FMorphTargetVertexSculptNonSymmetricChange;
-	virtual void RegisterBrushes() override
+	virtual void RegisterBrushes() override;
 
 	void OnToolMeshChanged(UDynamicMeshComponent* Component, const FMeshRegionChangeBase* Change, bool bRevert);
 
