[MorphEditingTools] Fixed reimportmesh removing all morphs created in engine.

diff --git a/Engine/Plugins/Experimental/Animation/SkeletalMeshMorphTargetEditingTools/Source/SkeletalMeshMorphTargetEditingTools/Private/SkeletalMeshMorphTargetToolTarget.cpp b/Engine/Plugins/Experimental/Animation/SkeletalMeshMorphTargetEditingTools/Source/SkeletalMeshMorphTargetEditingTools/Private/SkeletalMeshMorphTargetToolTarget.cpp
--- a/Engine/Plugins/Experimental/Animation/SkeletalMeshMorphTargetEditingTools/Source/SkeletalMeshMorphTargetEditingTools/Private/SkeletalMeshMorphTargetToolTarget.cpp
+++ b/Engine/Plugins/Experimental/Animation/SkeletalMeshMorphTargetEditingTools/Source/SkeletalMeshMorphTargetEditingTools/Private/SkeletalMeshMorphTargetToolTarget.cpp
@@ -14,6 +14,7 @@
 #include "ConversionUtils/SkinnedMeshToDynamicMesh.h"
 #include "Rendering/SkeletalMeshModel.h"
 #include "AnimationRuntime.h"
+#include "SkeletalMeshEditorSubsystem.h"
 
 UE::Geometry::FDynamicMesh3 USkeletalMeshMorphTargetToolTarget::GetDynamicMesh()
 {
@@ -60,6 +61,10 @@ void USkeletalMeshMorphTargetToolTarget::CommitDynamicMesh(const UE::Geometry::F
 	CommitEditedMeshDescription();
 
 	USkeletalMesh* SkeletalMesh = GetSkeletalMesh();
+	
+	// As soon as a morph is edited in engine, we want to mark it such that reimports in the future don't overwrite our edits.
+	USkeletalMeshEditorSubsystem::SetMorphTargetsToGeneratedByEngine(SkeletalMesh, {EditingMorphTargetName.ToString()});
+	
 	UAnimCurveMetaData* AnimCurveMetaData = SkeletalMesh->GetAssetUserData<UAnimCurveMetaData>();
 	if (AnimCurveMetaData == nullptr)
 	{
diff --git a/Engine/Plugins/Experimental/Animation/SkeletalMeshMorphTargetEditingTools/Source/SkeletalMeshMorphTargetEditingTools/SkeletalMeshMorphTargetEditingTools.Build.cs b/Engine/Plugins/Experimental/Animation/SkeletalMeshMorphTargetEditingTools/Source/SkeletalMeshMorphTargetEditingTools/SkeletalMeshMorphTargetEditingTools.Build.cs
--- a/Engine/Plugins/Experimental/Animation/SkeletalMeshMorphTargetEditingTools/Source/SkeletalMeshMorphTargetEditingTools/SkeletalMeshMorphTargetEditingTools.Build.cs
+++ b/Engine/Plugins/Experimental/Animation/SkeletalMeshMorphTargetEditingTools/Source/SkeletalMeshMorphTargetEditingTools/SkeletalMeshMorphTargetEditingTools.Build.cs
@@ -54,7 +54,8 @@ public class SkeletalMeshMorphTargetEditingTools : ModuleRules
 				"ModelingComponentsEditorOnly",
 				"UnrealEd",
 				"Persona",
-				"GeometryFramework"
+				"GeometryFramework", 
+				"SkeletalMeshEditor"
 				// ... add private dependencies that you statically link with here ...	
 			}
 			);
