[SkeletalMeshEditor] new icon for morph target + morph editing tool plugin + make file names shorter to avoid MAX_PATH limit

diff --git a/Engine/Plugins/Experimental/Animation/SkeletalMeshMorphTargetEditingTools/Source/SkeletalMeshMorphTargetEditingTools/Private/IMorphTargetEditingToolInterface.cpp b/Engine/Plugins/Experimental/Animation/SkeletalMeshMorphTargetEditingTools/Source/SkeletalMeshMorphTargetEditingTools/Private/IMorphTargetEditingToolInterface.cpp
--- a/Engine/Plugins/Experimental/Animation/SkeletalMeshMorphTargetEditingTools/Source/SkeletalMeshMorphTargetEditingTools/Private/IMorphTargetEditingToolInterface.cpp
+++ b/Engine/Plugins/Experimental/Animation/SkeletalMeshMorphTargetEditingTools/Source/SkeletalMeshMorphTargetEditingTools/Private/IMorphTargetEditingToolInterface.cpp
@@ -9,7 +9,7 @@
 #include "MorphTargetEditingToolProperties.h"
 #include "PersonaModule.h"
 #include "SingleSelectionTool.h"
-#include "SkeletalMeshMorphTargetBackedTarget.h"
+#include "SKMMorphTargetBackedTarget.h"
 #include "SkeletalMesh/SkeletalMeshEditionInterface.h"
 
 void IMorphTargetEditingToolInterface::SetupMorphEditingToolCommon()
diff --git a/Engine/Plugins/Experimental/Animation/SkeletalMeshMorphTargetEditingTools/Source/SkeletalMeshMorphTargetEditingTools/Private/MorphTargetVertexSculptTool.cpp b/Engine/Plugins/Experimental/Animation/SkeletalMeshMorphTargetEditingTools/Source/SkeletalMeshMorphTargetEditingTools/Private/MorphTargetVertexSculptTool.cpp
--- a/Engine/Plugins/Experimental/Animation/SkeletalMeshMorphTargetEditingTools/Source/SkeletalMeshMorphTargetEditingTools/Private/MorphTargetVertexSculptTool.cpp
+++ b/Engine/Plugins/Experimental/Animation/SkeletalMeshMorphTargetEditingTools/Source/SkeletalMeshMorphTargetEditingTools/Private/MorphTargetVertexSculptTool.cpp
@@ -15,7 +15,7 @@
 #include "DynamicMeshToMeshDescription.h"
 #include "IPersonaEditorModeManager.h"
 #include "PersonaModule.h"
-#include "SkeletalMeshMorphTargetBackedTarget.h"
+#include "SKMMorphTargetBackedTarget.h"
 #include "SkeletalMeshOperations.h"
 #include "BaseGizmos/BrushStampIndicator.h"
 #include "TargetInterfaces/DynamicMeshCommitter.h"
diff --git a/Engine/Plugins/Experimental/Animation/SkeletalMeshMorphTargetEditingTools/Source/SkeletalMeshMorphTargetEditingTools/Private/SKMMorphTargetEditingToolsCommands.cpp b/Engine/Plugins/Experimental/Animation/SkeletalMeshMorphTargetEditingTools/Source/SkeletalMeshMorphTargetEditingTools/Private/SKMMorphTargetEditingToolsCommands.cpp
new file mode 100644
--- /dev/null
+++ b/Engine/Plugins/Experimental/Animation/SkeletalMeshMorphTargetEditingTools/Source/SkeletalMeshMorphTargetEditingTools/Private/SKMMorphTargetEditingToolsCommands.cpp
@@ -0,0 +1,28 @@
+// Copyright Epic Games, Inc. All Rights Reserved.
+
+#include "SKMMorphTargetEditingToolsCommands.h"
+#include "SKMMorphTargetEditingToolsStyle.h"
+
+#define LOCTEXT_NAMESPACE "SkeletalMeshMorphTargetEditingToolsCommands"
+
+
+
+FSkeletalMeshMorphTargetEditingToolsCommands::FSkeletalMeshMorphTargetEditingToolsCommands()
+	: TCommands<FSkeletalMeshMorphTargetEditingToolsCommands>(
+		"SkeletalMeshMorphTargetEditingToolsCommands", // Context name for fast lookup
+		NSLOCTEXT("Contexts", "SkeletalMeshMorphTargetEditingToolsCommands", "Skeletal Mesh Morph Target Editing Tools"), // Localized context name for displaying
+		NAME_None, // Parent
+		FSkeletalMeshMorphTargetEditingToolsStyle::Get().GetStyleSetName() // Icon Style Set
+		)
+{
+}
+
+void FSkeletalMeshMorphTargetEditingToolsCommands::RegisterCommands()
+{
+	UI_COMMAND(BeginMorphTargetTool, "Morph", "Edit Morph Targets", EUserInterfaceActionType::ToggleButton, FInputChord());
+	
+	UI_COMMAND(BeginMorphTargetSculptTool, "Sculpt Morph Target", "MorphTarget Sculpt Tool", EUserInterfaceActionType::ToggleButton, FInputChord());
+}
+
+#undef LOCTEXT_NAMESPACE
+
diff --git a/Engine/Plugins/Experimental/Animation/SkeletalMeshMorphTargetEditingTools/Source/SkeletalMeshMorphTargetEditingTools/Private/SKMMorphTargetEditingToolsModule.cpp b/Engine/Plugins/Experimental/Animation/SkeletalMeshMorphTargetEditingTools/Source/SkeletalMeshMorphTargetEditingTools/Private/SKMMorphTargetEditingToolsModule.cpp
new file mode 100644
--- /dev/null
+++ b/Engine/Plugins/Experimental/Animation/SkeletalMeshMorphTargetEditingTools/Source/SkeletalMeshMorphTargetEditingTools/Private/SKMMorphTargetEditingToolsModule.cpp
@@ -0,0 +1,61 @@
+﻿// Copyright Epic Games, Inc. All Rights Reserved.
+
+#include "SKMMorphTargetEditingToolsModule.h"
+
+#include "MorphTargetVertexSculptTool.h"
+#include "SKMMorphTargetEditingToolsCommands.h"
+#include "SKMMorphTargetEditingToolsStyle.h"
+#include "SKMMorphTargetToolTarget.h"
+#include "Features/IModularFeatures.h"
+#include "Framework/Commands/Commands.h"
+
+#define LOCTEXT_NAMESPACE "FSkeletalMeshMorphTargetEditingToolsModule"
+
+void FSkeletalMeshMorphTargetEditingToolsModule::StartupModule()
+{
+	FSkeletalMeshMorphTargetEditingToolsStyle::Register();
+	FSkeletalMeshMorphTargetEditingToolsCommands::Register();
+	IModularFeatures::Get().RegisterModularFeature(ISkeletalMeshModelingModeToolExtension::GetModularFeatureName(), this);	
+
+}
+
+void FSkeletalMeshMorphTargetEditingToolsModule::ShutdownModule()
+{
+	IModularFeatures::Get().UnregisterModularFeature(ISkeletalMeshModelingModeToolExtension::GetModularFeatureName(), this);
+	FSkeletalMeshMorphTargetEditingToolsCommands::Unregister();
+	FSkeletalMeshMorphTargetEditingToolsStyle::Unregister();
+}
+
+
+
+void FSkeletalMeshMorphTargetEditingToolsModule::GetExtensionTools(const FExtensionToolQueryInfo& QueryInfo,
+                                                                   TArray<FExtensionToolDescription>& OutTools)
+{
+	const FSkeletalMeshMorphTargetEditingToolsCommands& Commands = FSkeletalMeshMorphTargetEditingToolsCommands::Get();
+	{
+		FExtensionToolDescription ToolDesc;
+		ToolDesc.ToolName = LOCTEXT("SkeletalMeshMorphTargetSculptTool", "Sculpt Morph Target");
+		ToolDesc.ToolCommand = Commands.BeginMorphTargetSculptTool;
+		ToolDesc.ToolBuilder = NewObject<UMorphTargetVertexSculptToolBuilder>();
+		OutTools.Add(ToolDesc);
+	}
+}
+
+bool FSkeletalMeshMorphTargetEditingToolsModule::GetExtensionExtendedInfo(FModelingModeExtensionExtendedInfo& InfoOut)
+{
+	InfoOut.ExtensionCommand = FSkeletalMeshMorphTargetEditingToolsCommands::Get().BeginMorphTargetTool;
+	
+	return true;
+}
+
+bool FSkeletalMeshMorphTargetEditingToolsModule::GetExtensionToolTargets(TArray<TSubclassOf<UToolTargetFactory>>& ToolTargetFactoriesOut)
+{
+	ToolTargetFactoriesOut.Add(USkeletalMeshMorphTargetToolTargetFactory::StaticClass());
+	
+	return true;
+}
+
+
+#undef LOCTEXT_NAMESPACE
+	
+IMPLEMENT_MODULE(FSkeletalMeshMorphTargetEditingToolsModule, SkeletalMeshMorphTargetEditingTools)
\ No newline at end of file
diff --git a/Engine/Plugins/Experimental/Animation/SkeletalMeshMorphTargetEditingTools/Source/SkeletalMeshMorphTargetEditingTools/Private/SKMMorphTargetEditingToolsStyle.cpp b/Engine/Plugins/Experimental/Animation/SkeletalMeshMorphTargetEditingTools/Source/SkeletalMeshMorphTargetEditingTools/Private/SKMMorphTargetEditingToolsStyle.cpp
new file mode 100644
--- /dev/null
+++ b/Engine/Plugins/Experimental/Animation/SkeletalMeshMorphTargetEditingTools/Source/SkeletalMeshMorphTargetEditingTools/Private/SKMMorphTargetEditingToolsStyle.cpp
@@ -0,0 +1,45 @@
+// Copyright Epic Games, Inc. All Rights Reserved.
+
+#include "SKMMorphTargetEditingToolsStyle.h"
+
+#include "Styling/SlateStyleMacros.h"
+#include "Styling/SlateStyleRegistry.h"
+
+
+FSkeletalMeshMorphTargetEditingToolsStyle& FSkeletalMeshMorphTargetEditingToolsStyle::Get()
+{
+	static FSkeletalMeshMorphTargetEditingToolsStyle Instance;
+	return Instance;
+}
+
+void FSkeletalMeshMorphTargetEditingToolsStyle::Register()
+{
+	FSlateStyleRegistry::RegisterSlateStyle(Get());
+}
+
+void FSkeletalMeshMorphTargetEditingToolsStyle::Unregister()
+{
+	FSlateStyleRegistry::UnRegisterSlateStyle(Get());
+}
+
+FSkeletalMeshMorphTargetEditingToolsStyle::FSkeletalMeshMorphTargetEditingToolsStyle()
+	: FSlateStyleSet("SkeletalMeshMorphTargetEditingToolsStyle")
+{
+	static const FVector2D IconSize10x10(10.0f, 10.0f);
+	static const FVector2D IconSize16x12(16.0f, 12.0f);
+	static const FVector2D IconSize16x16(16.0f, 16.0f);
+	static const FVector2D IconSize20x20(20.0f, 20.0f);
+	static const FVector2D IconSize64x64(64.0f, 64.0f);
+
+	static const FSlateColor DefaultForeground(FLinearColor(0.72f, 0.72f, 0.72f, 1.f));
+
+	{
+		FSlateStyleSet::SetContentRoot(FPaths::EnginePluginsDir() / TEXT("Experimental/Animation/SkeletalMeshMorphTargetEditingTools/Resources"));
+		Set("SkeletalMeshMorphTargetEditingToolsCommands.BeginMorphTargetTool", new IMAGE_BRUSH("Icon128", IconSize20x20));
+	}
+
+	{
+		FSlateStyleSet::SetContentRoot(FPaths::EnginePluginsDir() / TEXT("Editor/ModelingToolsEditorMode/Content"));
+		Set("SkeletalMeshMorphTargetEditingToolsCommands.BeginMorphTargetSculptTool", new IMAGE_BRUSH("Icons/Sculpt_40x", IconSize20x20));	
+	}
+}
diff --git a/Engine/Plugins/Experimental/Animation/SkeletalMeshMorphTargetEditingTools/Source/SkeletalMeshMorphTargetEditingTools/Private/SKMMorphTargetToolTarget.cpp b/Engine/Plugins/Experimental/Animation/SkeletalMeshMorphTargetEditingTools/Source/SkeletalMeshMorphTargetEditingTools/Private/SKMMorphTargetToolTarget.cpp
new file mode 100644
--- /dev/null
+++ b/Engine/Plugins/Experimental/Animation/SkeletalMeshMorphTargetEditingTools/Source/SkeletalMeshMorphTargetEditingTools/Private/SKMMorphTargetToolTarget.cpp
@@ -0,0 +1,170 @@
+﻿// Copyright Epic Games, Inc. All Rights Reserved.
+
+#include "SKMMorphTargetToolTarget.h"
+
+#include "DynamicMeshToMeshDescription.h"
+#include "MeshDescriptionToDynamicMesh.h"
+#include "SkeletalMeshAttributes.h"
+#include "SkeletalMeshOperations.h"
+#include "SkeletalMeshTypes.h"
+#include "StaticMeshAttributes.h"
+#include "Components/SkeletalMeshComponent.h"
+#include "ConversionUtils/DynamicMeshViaMeshDescriptionUtil.h"
+#include "UObject/Package.h"
+#include "ConversionUtils/SkinnedMeshToDynamicMesh.h"
+#include "Rendering/SkeletalMeshModel.h"
+#include "AnimationRuntime.h"
+#include "SkeletalMeshEditorSubsystem.h"
+
+UE::Geometry::FDynamicMesh3 USkeletalMeshMorphTargetToolTarget::GetDynamicMesh()
+{
+	FGetMeshParameters Params;
+	return GetDynamicMesh(Params);
+}
+
+UE::Geometry::FDynamicMesh3 USkeletalMeshMorphTargetToolTarget::GetDynamicMesh(const FGetMeshParameters& InGetMeshParams)
+{
+	FDynamicMesh3 DynamicMesh;
+	FMeshDescriptionToDynamicMesh Converter;
+	Converter.Convert(GetEditingMeshDescription(), DynamicMesh);
+	
+	return DynamicMesh;
+}
+
+void USkeletalMeshMorphTargetToolTarget::CommitDynamicMesh(const UE::Geometry::FDynamicMesh3& Mesh, const FDynamicMeshCommitInfo& CommitInfo)
+{
+	if (ToolMeshDescription.IsEmpty() || ToolMorphTargetName.IsNone() || EditingMorphTargetName.IsNone())
+	{
+		return;
+	}
+	
+	FMeshDescription* TargetMeshDescription = GetEditingMeshDescription();
+	FSkeletalMeshAttributes TargetAttributes(*TargetMeshDescription);
+
+	if (!TargetAttributes.GetMorphTargetNames().Contains(EditingMorphTargetName))
+	{
+		TargetAttributes.RegisterMorphTargetAttribute(EditingMorphTargetName, false);
+	}
+	
+	TVertexAttributesRef<FVector3f> TargetMorphTargetPosDeltaAttribute = TargetAttributes.GetVertexMorphPositionDelta(EditingMorphTargetName);
+
+	
+	FSkeletalMeshAttributes ToolAttributes(ToolMeshDescription);
+	TVertexAttributesRef<FVector3f> ToolMorphTargetPosDeltaAttribute = ToolAttributes.GetVertexMorphPositionDelta(ToolMorphTargetName);
+	
+	
+	for (FVertexID VertexID : TargetMeshDescription->Vertices().GetElementIDs())
+	{
+		TargetMorphTargetPosDeltaAttribute[VertexID] = ToolMorphTargetPosDeltaAttribute[VertexID.GetValue()];
+	}		
+
+	CommitEditedMeshDescription();
+
+	USkeletalMesh* SkeletalMesh = GetSkeletalMesh();
+	
+	// As soon as a morph is edited in engine, we want to mark it such that reimports in the future don't overwrite our edits.
+	USkeletalMeshEditorSubsystem::SetMorphTargetsToGeneratedByEngine(SkeletalMesh, {EditingMorphTargetName.ToString()});
+	
+	UAnimCurveMetaData* AnimCurveMetaData = SkeletalMesh->GetAssetUserData<UAnimCurveMetaData>();
+	if (AnimCurveMetaData == nullptr)
+	{
+		AnimCurveMetaData = NewObject<UAnimCurveMetaData>(SkeletalMesh, NAME_None, RF_Transactional);
+		SkeletalMesh->AddAssetUserData(AnimCurveMetaData);
+	}
+
+	AnimCurveMetaData->AddCurveMetaData(EditingMorphTargetName);
+
+	// Ensure we have a morph flag set
+	if (FCurveMetaData* CurveMetaData = AnimCurveMetaData->GetCurveMetaData(EditingMorphTargetName))
+	{
+		CurveMetaData->Type.bMorphtarget = true;
+	}
+}
+
+void USkeletalMeshMorphTargetToolTarget::SetDataToCommit(FMeshDescription&& InMeshDescription, const FName& InToolMorphTargetName)
+{
+	ToolMeshDescription = MoveTemp(InMeshDescription);
+	ToolMorphTargetName = InToolMorphTargetName;
+}
+
+void USkeletalMeshMorphTargetToolTarget::SetEditingMorphTargetName(const FName& InName)
+{
+	EditingMorphTargetName = InName;
+}
+
+TArray<FName> USkeletalMeshMorphTargetToolTarget::GetEditableMorphTargetNames()
+{
+	FSkeletalMeshAttributes TargetAttributes(*GetEditingMeshDescription());
+
+	return TargetAttributes.GetMorphTargetNames();
+}
+
+FName USkeletalMeshMorphTargetToolTarget::GetValidNameForNewMorphTarget(const FName& InName)
+{
+	FSkeletalMeshAttributes TargetAttributes(*GetEditingMeshDescription());
+
+	const TArray<FName>& MorphTargetNames = TargetAttributes.GetMorphTargetNames();
+	FName NewMorphTargetName = InName;
+	while (MorphTargetNames.Contains(NewMorphTargetName))
+	{
+		NewMorphTargetName.SetNumber(NewMorphTargetName.GetNumber() + 1);
+	}
+
+	return NewMorphTargetName;
+}
+
+FMeshDescription* USkeletalMeshMorphTargetToolTarget::GetEditingMeshDescription()
+{
+	USkeletalMesh* SkeletalMesh = GetSkeletalMesh();
+	FMeshDescription* TargetMeshDescription = SkeletalMesh->GetMeshDescription(LodIndex0);
+
+	return TargetMeshDescription;
+}
+
+void USkeletalMeshMorphTargetToolTarget::CommitEditedMeshDescription()
+{
+	USkeletalMesh* SkeletalMesh = GetSkeletalMesh();
+	{
+		FScopedSkeletalMeshPostEditChange ScopePostEditChange(SkeletalMesh);
+		SkeletalMesh->PreEditChange(nullptr);
+		SkeletalMesh->ModifyMeshDescription(LodIndex0);
+		SkeletalMesh->CommitMeshDescription(LodIndex0);
+	}
+}
+
+bool USkeletalMeshMorphTargetToolTargetFactory::CanBuildTarget(UObject* SourceObject,
+                                                               const FToolTargetTypeRequirements& Requirements) const
+{
+	// We are using an exact cast here to prevent subclasses, which might not meet all
+	// requirements for functionality such as the deprecated DestructibleMeshComponent, from 
+	// being caught up as valid targets.
+	// If you want to make the tool target work with some subclass of USkeletalMeshComponent,
+	// just add another factory that allows that class specifically(but make sure that
+	// GetMeshDescription and such work properly)
+	
+	bool bValid = Cast<USkeletalMeshComponent>(SourceObject) && Cast<USkeletalMeshComponent>(SourceObject)->GetSkeletalMeshAsset() &&
+		ExactCast<USkeletalMesh>(Cast<USkeletalMeshComponent>(SourceObject)->GetSkeletalMeshAsset()) &&
+		!ExactCast<USkeletalMesh>(Cast<USkeletalMeshComponent>(SourceObject)->GetSkeletalMeshAsset())->GetOutermost()->bIsCookedForEditor;
+	if (!bValid)
+	{
+		return false;
+	}
+
+	if (!USkeletalMeshComponentToolTargetFactory::CanWriteToSource(SourceObject))
+	{
+		return false;
+	}
+
+	return Requirements.AreSatisfiedBy(USkeletalMeshMorphTargetToolTarget::StaticClass());
+}
+
+UToolTarget* USkeletalMeshMorphTargetToolTargetFactory::BuildTarget(UObject* SourceObject, const FToolTargetTypeRequirements& TargetTypeInfo)
+{
+	USkeletalMeshMorphTargetToolTarget* Target = NewObject<USkeletalMeshMorphTargetToolTarget>();
+
+	Target->InitializeComponent(Cast<USkeletalMeshComponent>(SourceObject));
+	
+	return Target;
+	
+}
+
diff --git a/Engine/Plugins/Experimental/Animation/SkeletalMeshMorphTargetEditingTools/Source/SkeletalMeshMorphTargetEditingTools/Private/SkeletalMeshMorphTargetEditingToolsCommands.cpp b/Engine/Plugins/Experimental/Animation/SkeletalMeshMorphTargetEditingTools/Source/SkeletalMeshMorphTargetEditingTools/Private/SkeletalMeshMorphTargetEditingToolsCommands.cpp
deleted file mode 100644
diff --git a/Engine/Plugins/Experimental/Animation/SkeletalMeshMorphTargetEditingTools/Source/SkeletalMeshMorphTargetEditingTools/Private/SkeletalMeshMorphTargetEditingToolsModule.cpp b/Engine/Plugins/Experimental/Animation/SkeletalMeshMorphTargetEditingTools/Source/SkeletalMeshMorphTargetEditingTools/Private/SkeletalMeshMorphTargetEditingToolsModule.cpp
deleted file mode 100644
diff --git a/Engine/Plugins/Experimental/Animation/SkeletalMeshMorphTargetEditingTools/Source/SkeletalMeshMorphTargetEditingTools/Private/SkeletalMeshMorphTargetEditingToolsStyle.cpp b/Engine/Plugins/Experimental/Animation/SkeletalMeshMorphTargetEditingTools/Source/SkeletalMeshMorphTargetEditingTools/Private/SkeletalMeshMorphTargetEditingToolsStyle.cpp
deleted file mode 100644
diff --git a/Engine/Plugins/Experimental/Animation/SkeletalMeshMorphTargetEditingTools/Source/SkeletalMeshMorphTargetEditingTools/Private/SkeletalMeshMorphTargetToolTarget.cpp b/Engine/Plugins/Experimental/Animation/SkeletalMeshMorphTargetEditingTools/Source/SkeletalMeshMorphTargetEditingTools/Private/SkeletalMeshMorphTargetToolTarget.cpp
deleted file mode 100644
diff --git a/Engine/Plugins/Experimental/Animation/SkeletalMeshMorphTargetEditingTools/Source/SkeletalMeshMorphTargetEditingTools/Public/MorphTargetVertexSculptTool.h b/Engine/Plugins/Experimental/Animation/SkeletalMeshMorphTargetEditingTools/Source/SkeletalMeshMorphTargetEditingTools/Public/MorphTargetVertexSculptTool.h
--- a/Engine/Plugins/Experimental/Animation/SkeletalMeshMorphTargetEditingTools/Source/SkeletalMeshMorphTargetEditingTools/Public/MorphTargetVertexSculptTool.h
+++ b/Engine/Plugins/Experimental/Animation/SkeletalMeshMorphTargetEditingTools/Source/SkeletalMeshMorphTargetEditingTools/Public/MorphTargetVertexSculptTool.h
@@ -7,7 +7,7 @@
 #include "SkeletalMesh/SkeletalMeshEditionInterface.h"
 #include "BaseTools/MultiSelectionMeshEditingTool.h"
 #include "MorphTargetEditingToolProperties.h"
-#include "SkeletalMeshMorphTargetBackedTarget.h"
+#include "SKMMorphTargetBackedTarget.h"
 #include "MeshDescription.h"
 
 #include "MorphTargetVertexSculptTool.generated.h"
diff --git a/Engine/Plugins/Experimental/Animation/SkeletalMeshMorphTargetEditingTools/Source/SkeletalMeshMorphTargetEditingTools/Public/SKMMorphTargetBackedTarget.h b/Engine/Plugins/Experimental/Animation/SkeletalMeshMorphTargetEditingTools/Source/SkeletalMeshMorphTargetEditingTools/Public/SKMMorphTargetBackedTarget.h
new file mode 100644
--- /dev/null
+++ b/Engine/Plugins/Experimental/Animation/SkeletalMeshMorphTargetEditingTools/Source/SkeletalMeshMorphTargetEditingTools/Public/SKMMorphTargetBackedTarget.h
@@ -0,0 +1,27 @@
+﻿// Copyright Epic Games, Inc. All Rights Reserved.
+
+#pragma once
+
+#include "CoreMinimal.h"
+#include "UObject/Interface.h"
+
+#include "SKMMorphTargetBackedTarget.generated.h"
+
+struct FMeshDescription;
+
+UINTERFACE(MinimalAPI)
+class USkeletalMeshMorphTargetBackedTarget : public UInterface
+{
+	GENERATED_BODY()
+};
+
+class ISkeletalMeshMorphTargetBackedTarget : public IInterface
+{
+	GENERATED_BODY()
+public:
+	virtual void SetDataToCommit(FMeshDescription&& InMeshDescription, const FName& ToolMorphTargetName) = 0;
+	virtual void SetEditingMorphTargetName(const FName& InName) = 0;
+	virtual TArray<FName> GetEditableMorphTargetNames() = 0;
+	virtual FName GetValidNameForNewMorphTarget(const FName& InName) = 0;
+};
+
diff --git a/Engine/Plugins/Experimental/Animation/SkeletalMeshMorphTargetEditingTools/Source/SkeletalMeshMorphTargetEditingTools/Public/SKMMorphTargetEditingToolsCommands.h b/Engine/Plugins/Experimental/Animation/SkeletalMeshMorphTargetEditingTools/Source/SkeletalMeshMorphTargetEditingTools/Public/SKMMorphTargetEditingToolsCommands.h
new file mode 100644
--- /dev/null
+++ b/Engine/Plugins/Experimental/Animation/SkeletalMeshMorphTargetEditingTools/Source/SkeletalMeshMorphTargetEditingTools/Public/SKMMorphTargetEditingToolsCommands.h
@@ -0,0 +1,18 @@
+// Copyright Epic Games, Inc. All Rights Reserved.
+
+#pragma once
+
+#include "CoreMinimal.h"
+#include "Framework/Commands/Commands.h"
+
+class FSkeletalMeshMorphTargetEditingToolsCommands : public TCommands<FSkeletalMeshMorphTargetEditingToolsCommands>
+{
+public:
+	FSkeletalMeshMorphTargetEditingToolsCommands();
+
+
+	TSharedPtr<FUICommandInfo> BeginMorphTargetTool;
+	TSharedPtr<FUICommandInfo> BeginMorphTargetSculptTool;
+
+	virtual void RegisterCommands() override;
+};
diff --git a/Engine/Plugins/Experimental/Animation/SkeletalMeshMorphTargetEditingTools/Source/SkeletalMeshMorphTargetEditingTools/Public/SKMMorphTargetEditingToolsModule.h b/Engine/Plugins/Experimental/Animation/SkeletalMeshMorphTargetEditingTools/Source/SkeletalMeshMorphTargetEditingTools/Public/SKMMorphTargetEditingToolsModule.h
new file mode 100644
--- /dev/null
+++ b/Engine/Plugins/Experimental/Animation/SkeletalMeshMorphTargetEditingTools/Source/SkeletalMeshMorphTargetEditingTools/Public/SKMMorphTargetEditingToolsModule.h
@@ -0,0 +1,28 @@
+﻿// Copyright Epic Games, Inc. All Rights Reserved.
+
+#pragma once
+
+#include "Modules/ModuleManager.h"
+#include "SkeletalMeshModelingModeToolExtensions.h"
+
+class FUICommandInfo;
+
+class FSkeletalMeshMorphTargetEditingToolsModule : public IModuleInterface, public ISkeletalMeshModelingModeToolExtension 
+{
+public:
+	/** IModuleInterface implementation */
+	virtual void StartupModule() override;
+	virtual void ShutdownModule() override;
+	
+	// ISkeletalMeshModelingModeToolExtension
+	virtual FText GetExtensionName() override { return FText::FromString(TEXT("SkeletalMeshMorphTargetEditingTools"));}
+
+	virtual FText GetToolSectionName() override { return FText::FromString(TEXT("Morph"));};
+
+	virtual void GetExtensionTools(const FExtensionToolQueryInfo& QueryInfo, TArray<FExtensionToolDescription>& OutTools) override;
+	
+	virtual bool GetExtensionExtendedInfo(FModelingModeExtensionExtendedInfo& InfoOut) override;
+
+	virtual bool GetExtensionToolTargets(TArray<TSubclassOf<UToolTargetFactory>>& ToolTargetFactoriesOut) override;
+
+};
diff --git a/Engine/Plugins/Experimental/Animation/SkeletalMeshMorphTargetEditingTools/Source/SkeletalMeshMorphTargetEditingTools/Public/SKMMorphTargetEditingToolsStyle.h b/Engine/Plugins/Experimental/Animation/SkeletalMeshMorphTargetEditingTools/Source/SkeletalMeshMorphTargetEditingTools/Public/SKMMorphTargetEditingToolsStyle.h
new file mode 100644
--- /dev/null
+++ b/Engine/Plugins/Experimental/Animation/SkeletalMeshMorphTargetEditingTools/Source/SkeletalMeshMorphTargetEditingTools/Public/SKMMorphTargetEditingToolsStyle.h
@@ -0,0 +1,22 @@
+// Copyright Epic Games, Inc. All Rights Reserved.
+
+#pragma once
+
+#include "Styling/SlateStyle.h"
+
+
+class FSkeletalMeshMorphTargetEditingToolsStyle
+	: public FSlateStyleSet
+{
+public:
+	static FSkeletalMeshMorphTargetEditingToolsStyle& Get();
+
+protected:
+	friend class FSkeletalMeshMorphTargetEditingToolsModule;
+
+	static void Register();
+	static void Unregister();
+
+private:
+	FSkeletalMeshMorphTargetEditingToolsStyle();
+};
\ No newline at end of file
diff --git a/Engine/Plugins/Experimental/Animation/SkeletalMeshMorphTargetEditingTools/Source/SkeletalMeshMorphTargetEditingTools/Public/SKMMorphTargetToolTarget.h b/Engine/Plugins/Experimental/Animation/SkeletalMeshMorphTargetEditingTools/Source/SkeletalMeshMorphTargetEditingTools/Public/SKMMorphTargetToolTarget.h
new file mode 100644
--- /dev/null
+++ b/Engine/Plugins/Experimental/Animation/SkeletalMeshMorphTargetEditingTools/Source/SkeletalMeshMorphTargetEditingTools/Public/SKMMorphTargetToolTarget.h
@@ -0,0 +1,70 @@
+﻿// Copyright Epic Games, Inc. All Rights Reserved.
+
+#pragma once
+
+#include "CoreMinimal.h"
+
+#include "ToolTargets/SkeletalMeshComponentToolTarget.h"
+#include "SKMMorphTargetBackedTarget.h"
+#include "ToolTargets/ToolTarget.h"
+#include "MeshDescription.h"
+
+#include "SKMMorphTargetToolTarget.generated.h"
+
+/**
+ * A tool target backed by a skeletal mesh component that can provide and take a mesh
+ * description.
+ */
+UCLASS(Transient)
+class USkeletalMeshMorphTargetToolTarget :
+	public USkeletalMeshComponentReadOnlyToolTarget,
+	public IDynamicMeshCommitter,
+	public ISkeletalMeshMorphTargetBackedTarget
+{
+	GENERATED_BODY()
+
+public:
+	virtual UE::Geometry::FDynamicMesh3 GetDynamicMesh() override;
+	virtual UE::Geometry::FDynamicMesh3 GetDynamicMesh(const FGetMeshParameters& InGetMeshParams) override;
+	
+	// IDynamicMeshCommitter
+	virtual void CommitDynamicMesh(const UE::Geometry::FDynamicMesh3& Mesh, const FDynamicMeshCommitInfo& CommitInfo) override;
+	using IDynamicMeshCommitter::CommitDynamicMesh; // unhide the other overload
+
+
+	virtual void SetDataToCommit(FMeshDescription&& InMeshDescription, const FName& InToolMorphTargetName) override;
+	virtual void SetEditingMorphTargetName(const FName& InName) override;
+	virtual TArray<FName> GetEditableMorphTargetNames() override;
+	virtual FName GetValidNameForNewMorphTarget(const FName& InName) override;
+protected:
+	// So that the tool target factory can poke into Component.
+	friend class USkeletalMeshMorphTargetToolTargetFactory;
+
+	FMeshDescription* GetEditingMeshDescription();
+	void CommitEditedMeshDescription();
+	
+	UPROPERTY()
+	FName EditingMorphTargetName;
+
+	static constexpr int32 LodIndex0 = 0;
+
+	FMeshDescription ToolMeshDescription;
+	FName ToolMorphTargetName;
+};
+
+
+
+
+UCLASS(Transient)
+class USkeletalMeshMorphTargetToolTargetFactory : public UToolTargetFactory
+{
+	GENERATED_BODY()
+
+public:
+
+	virtual bool CanBuildTarget(UObject* SourceObject, const FToolTargetTypeRequirements& Requirements) const override;
+
+	virtual UToolTarget* BuildTarget(UObject* SourceObject, const FToolTargetTypeRequirements& TargetTypeInfo) override;
+};
+
+
diff --git a/Engine/Plugins/Experimental/Animation/SkeletalMeshMorphTargetEditingTools/Source/SkeletalMeshMorphTargetEditingTools/Public/SkeletalMeshMorphTargetBackedTarget.h b/Engine/Plugins/Experimental/Animation/SkeletalMeshMorphTargetEditingTools/Source/SkeletalMeshMorphTargetEditingTools/Public/SkeletalMeshMorphTargetBackedTarget.h
deleted file mode 100644
diff --git a/Engine/Plugins/Experimental/Animation/SkeletalMeshMorphTargetEditingTools/Source/SkeletalMeshMorphTargetEditingTools/Public/SkeletalMeshMorphTargetEditingToolsCommands.h b/Engine/Plugins/Experimental/Animation/SkeletalMeshMorphTargetEditingTools/Source/SkeletalMeshMorphTargetEditingTools/Public/SkeletalMeshMorphTargetEditingToolsCommands.h
deleted file mode 100644
diff --git a/Engine/Plugins/Experimental/Animation/SkeletalMeshMorphTargetEditingTools/Source/SkeletalMeshMorphTargetEditingTools/Public/SkeletalMeshMorphTargetEditingToolsModule.h b/Engine/Plugins/Experimental/Animation/SkeletalMeshMorphTargetEditingTools/Source/SkeletalMeshMorphTargetEditingTools/Public/SkeletalMeshMorphTargetEditingToolsModule.h
deleted file mode 100644
diff --git a/Engine/Plugins/Experimental/Animation/SkeletalMeshMorphTargetEditingTools/Source/SkeletalMeshMorphTargetEditingTools/Public/SkeletalMeshMorphTargetEditingToolsStyle.h b/Engine/Plugins/Experimental/Animation/SkeletalMeshMorphTargetEditingTools/Source/SkeletalMeshMorphTargetEditingTools/Public/SkeletalMeshMorphTargetEditingToolsStyle.h
deleted file mode 100644
diff --git a/Engine/Plugins/Experimental/Animation/SkeletalMeshMorphTargetEditingTools/Source/SkeletalMeshMorphTargetEditingTools/Public/SkeletalMeshMorphTargetToolTarget.h b/Engine/Plugins/Experimental/Animation/SkeletalMeshMorphTargetEditingTools/Source/SkeletalMeshMorphTargetEditingTools/Public/SkeletalMeshMorphTargetToolTarget.h
deleted file mode 100644
diff --git a/Engine/Source/Editor/EditorStyle/Private/StarshipStyle.cpp b/Engine/Source/Editor/EditorStyle/Private/StarshipStyle.cpp
--- a/Engine/Source/Editor/EditorStyle/Private/StarshipStyle.cpp
+++ b/Engine/Source/Editor/EditorStyle/Private/StarshipStyle.cpp
@@ -5572,7 +5572,7 @@ void FStarshipEditorStyle::FStyle::SetupPersonaStyle()
 
 		// Persona-specific tabs
 		Set("Persona.Tabs.SkeletonTree", new IMAGE_BRUSH(TEXT("Icons/icon_Persona_Skeleton_Tree_16x"), Icon16x16));
-		Set("Persona.Tabs.MorphTargetPreviewer", new IMAGE_BRUSH(TEXT("Icons/icon_Persona_Morph_Target_Previewer_16x"), Icon16x16));
+		Set("Persona.Tabs.MorphTargetPreviewer", new IMAGE_BRUSH_SVG(TEXT("Starship/Persona/MorphTarget"), Icon16x16));
 		Set("Persona.Tabs.AnimCurveDebugger", new IMAGE_BRUSH_SVG(TEXT("Starship/Persona/CurveDebugger"), Icon16x16));
 		Set("Persona.Tabs.AnimCurveMetadataEditor", new IMAGE_BRUSH_SVG(TEXT("Starship/Persona/CurveMetadata"), Icon16x16));
 		Set("Persona.Tabs.AnimationNotifies", new IMAGE_BRUSH(TEXT("Icons/icon_Persona_Animation_Notifies_16x"), Icon16x16));
