[MorphTargetSculptTool] Fixed octree out of date when undoing a brush stroke

diff --git a/Engine/Plugins/Experimental/Animation/SkeletalMeshMorphTargetEditingTools/Source/SkeletalMeshMorphTargetEditingTools/Private/MorphTargetVertexSculptTool.cpp b/Engine/Plugins/Experimental/Animation/SkeletalMeshMorphTargetEditingTools/Source/SkeletalMeshMorphTargetEditingTools/Private/MorphTargetVertexSculptTool.cpp
--- a/Engine/Plugins/Experimental/Animation/SkeletalMeshMorphTargetEditingTools/Source/SkeletalMeshMorphTargetEditingTools/Private/MorphTargetVertexSculptTool.cpp
+++ b/Engine/Plugins/Experimental/Animation/SkeletalMeshMorphTargetEditingTools/Source/SkeletalMeshMorphTargetEditingTools/Private/MorphTargetVertexSculptTool.cpp
@@ -114,11 +114,6 @@ void UMorphTargetVertexSculptTool::ToggleBoneManipulation(bool bEnable)
 	}
 }
 
-void UMorphTargetVertexSculptTool::RequestRefreshSculptMesh()
-{
-	bSculptMeshOutdated = true;
-}
-
 void UMorphTargetVertexSculptTool::HandlePoseChangeDetectorEvent(SkeletalMeshToolsHelper::FPoseChangeDetector::FPayload Payload)
 {
 	if (!AllowToolMeshUpdates())
@@ -129,19 +124,12 @@ void UMorphTargetVertexSculptTool::HandlePoseChangeDetectorEvent(SkeletalMeshToo
 	using namespace SkeletalMeshToolsHelper;
 	if (Payload.CurrentState == FPoseChangeDetector::PoseStoppedChanging)
 	{
-		FMeshReplacementChange Change(MeshBeforePosing, MakeShared<FDynamicMesh3>(*GetSculptMesh()));
-		TGuardValue<bool> PosingSculptMeshScope(bPosingSculptMesh, true);
-		OnDynamicMeshComponentChanged(DynamicMeshComponent, &Change, false);	
+		bFullRefreshSculptMesh = true;
 	}
 	else if (Payload.CurrentState == FPoseChangeDetector::PoseJustChanged ||
 		Payload.CurrentState == FPoseChangeDetector::PoseChanged)
 	{
-		if (Payload.CurrentState == FPoseChangeDetector::PoseJustChanged)
-		{
-			MeshBeforePosing = MakeShared<FDynamicMesh3>(*GetSculptMesh());		
-		}
-
-		RequestRefreshSculptMesh();
+		bFastDeformSculptMesh = true;
 	}
 }
 
@@ -231,11 +219,22 @@ void UMorphTargetVertexSculptTool::OnTick(float DeltaTime)
 		
 		PoseChangeDetector.CheckPose(ComponentSpaceTransforms, MorphTargetWeights);
 
-		if (bSculptMeshOutdated)
+		if (bFastDeformSculptMesh)
 		{
-			bSculptMeshOutdated = false;
+			bFastDeformSculptMesh = false;
+			
 			PoseSculptMesh(ComponentSpaceTransforms, MorphTargetWeights);
 		}
+
+		if (bFullRefreshSculptMesh)
+		{
+			bFullRefreshSculptMesh = false;
+			
+			// Using a dummy mesh replacement change to trigger octree/normals/base mesh refresh
+			TSharedPtr<FDynamicMesh3> DummySculptMesh = MakeShared<FDynamicMesh3>(*GetSculptMesh());	
+			FMeshReplacementChange DummyChange(DummySculptMesh, DummySculptMesh);
+			OnDynamicMeshComponentChanged(DynamicMeshComponent, &DummyChange, false);
+		}
 	}
 }
 
@@ -366,11 +365,14 @@ void UMorphTargetVertexSculptTool::OnEndStroke()
 		NewChange->WrappedChange.Reset(MorphTargetChange);
 		NewChange->BeforeModify = [this](bool bRevert)
 		{
-			this->WaitForPendingUndoRedoUpdate();
+			WaitForPendingUndoRedoUpdate();
 		};
-		NewChange->AfterModify= [this](bool bRevert)
+
+		
+		NewChange->AfterModify = [this](bool bRevert)
 		{
-			this->RequestRefreshSculptMesh();
+			bFastDeformSculptMesh = true;
+			bFullRefreshSculptMesh = true;
 		};
 
 		GetToolManager()->EmitObjectChange(ToolDynamicMesh, MoveTemp(NewChange), LOCTEXT("VertexSculptChange", "Brush Stroke"));
@@ -381,8 +383,6 @@ void UMorphTargetVertexSculptTool::OnEndStroke()
 			bMeshSymmetryIsValid = false;
 			SymmetryProperties->bSymmetryCanBeEnabled = bMeshSymmetryIsValid;
 		}
-
-		RequestRefreshSculptMesh();
 	}
 	
 	
diff --git a/Engine/Plugins/Experimental/Animation/SkeletalMeshMorphTargetEditingTools/Source/SkeletalMeshMorphTargetEditingTools/Public/MorphTargetVertexSculptTool.h b/Engine/Plugins/Experimental/Animation/SkeletalMeshMorphTargetEditingTools/Source/SkeletalMeshMorphTargetEditingTools/Public/MorphTargetVertexSculptTool.h
--- a/Engine/Plugins/Experimental/Animation/SkeletalMeshMorphTargetEditingTools/Source/SkeletalMeshMorphTargetEditingTools/Public/MorphTargetVertexSculptTool.h
+++ b/Engine/Plugins/Experimental/Animation/SkeletalMeshMorphTargetEditingTools/Source/SkeletalMeshMorphTargetEditingTools/Public/MorphTargetVertexSculptTool.h
@@ -45,8 +45,6 @@ public:
 	virtual const TArray<FTransform>& GetComponentSpaceBoneTransforms();
 	virtual void ToggleBoneManipulation(bool bEnable);
 
-	void RequestRefreshSculptMesh();
-
 
 	// UMeshVertexSculptTool overrides
 	virtual void Setup() override;
@@ -83,7 +81,7 @@ protected:
 	
 	SkeletalMeshToolsHelper::FPoseChangeDetector PoseChangeDetector;
 
-	TSharedPtr<FDynamicMesh3> MeshBeforePosing;
+	
 
 	FDelegateHandle OnToolMeshChangedDelegate;
 
@@ -94,14 +92,14 @@ protected:
 	
 	FName ToolMorphTargetName;
 
-	bool bPosingSculptMesh = false;
 
 	TArray<FTransform> DefaultComponentSpaceBoneTransforms;
 	
 	TArray<FTransform> ComponentSpaceTransformsForPosedMesh;
 	TMap<FName, float> MorphTargetWeightsForPosedMesh;
 
-	bool bSculptMeshOutdated = false;
+	bool bFastDeformSculptMesh = false;
+	bool bFullRefreshSculptMesh = false;
 };
 
 
