[MorphSculptTool] Added array size check to avoid crash when dealing with a null stroke.

diff --git a/Engine/Plugins/Experimental/Animation/SkeletalMeshMorphTargetEditingTools/Source/SkeletalMeshMorphTargetEditingTools/Private/MorphTargetVertexSculptTool.cpp b/Engine/Plugins/Experimental/Animation/SkeletalMeshMorphTargetEditingTools/Source/SkeletalMeshMorphTargetEditingTools/Private/MorphTargetVertexSculptTool.cpp
--- a/Engine/Plugins/Experimental/Animation/SkeletalMeshMorphTargetEditingTools/Source/SkeletalMeshMorphTargetEditingTools/Private/MorphTargetVertexSculptTool.cpp
+++ b/Engine/Plugins/Experimental/Animation/SkeletalMeshMorphTargetEditingTools/Source/SkeletalMeshMorphTargetEditingTools/Private/MorphTargetVertexSculptTool.cpp
@@ -321,7 +321,7 @@ void UMorphTargetVertexSculptTool::OnEndStroke()
 	float EditingMorphTargetWeight = MorphTargetWeights[ToolMorphTargetName];
 	
 	// Don't update the cache if the morph cannot be extracted from the mesh
-	if (!FMath::IsNearlyZero(EditingMorphTargetWeight))
+	if (!FMath::IsNearlyZero(EditingMorphTargetWeight) && !ActiveVertexChange->Change->Vertices.IsEmpty())
 	{
 		FMeshMorphTargetChange* MorphTargetChange = new FMeshMorphTargetChange(); 
 		MorphTargetChange->MorphTargetName = ToolMorphTargetName;
