[SKM morph sculpt editing tool]  hide brush alpha for brushes not using alpha

#Synced-CL 50243241

diff --git a/Engine/Plugins/Experimental/Animation/SkeletalMeshMorphTargetEditingTools/Source/SkeletalMeshMorphTargetEditingTools/Private/MorphTargetVertexSculptTool.cpp b/Engine/Plugins/Experimental/Animation/SkeletalMeshMorphTargetEditingTools/Source/SkeletalMeshMorphTargetEditingTools/Private/MorphTargetVertexSculptTool.cpp
--- a/Engine/Plugins/Experimental/Animation/SkeletalMeshMorphTargetEditingTools/Source/SkeletalMeshMorphTargetEditingTools/Private/MorphTargetVertexSculptTool.cpp
+++ b/Engine/Plugins/Experimental/Animation/SkeletalMeshMorphTargetEditingTools/Source/SkeletalMeshMorphTargetEditingTools/Private/MorphTargetVertexSculptTool.cpp
@@ -22,6 +22,7 @@
 #include "TargetInterfaces/PrimitiveComponentBackedTarget.h"
 #include "AnimationRuntime.h"
 #include "EraseMorphTargetBrushOps.h"
+#include "PreviewMesh.h"
 #include "SkeletalMeshAttributes.h"
 #include "DynamicMesh/DynamicVertexSkinWeightsAttribute.h"
 #include "SkeletalMesh/SkeletalMeshToolsHelper.h"
@@ -236,6 +237,9 @@ void UMorphTargetVertexSculptTool::OnTick(float DeltaTime)
 			OnDynamicMeshComponentChanged(DynamicMeshComponent, &DummyChange, false);
 		}
 	}
+
+	// Hide the brush alpha mesh when alpha is not used because it makes it hard to see the sculpt area clearly.
+	BrushIndicatorMesh->SetVisible(AlphaProperties->IsPropertySetEnabled() && AlphaProperties->Alpha != nullptr);
 }
 
 void UMorphTargetVertexSculptTool::OnPropertyModified(UObject* PropertySet, FProperty* Property)
