[Skeletal Mesh Modeling Tools] Added MorphTargetMaskTool

#Synced-CL 51489786

diff --git a/Engine/Plugins/Animation/SkeletalMeshModelingTools/Source/SkeletalMeshModelingTools/Private/SkeletalMeshEditingCache.cpp b/Engine/Plugins/Animation/SkeletalMeshModelingTools/Source/SkeletalMeshModelingTools/Private/SkeletalMeshEditingCache.cpp
--- a/Engine/Plugins/Animation/SkeletalMeshModelingTools/Source/SkeletalMeshModelingTools/Private/SkeletalMeshEditingCache.cpp
+++ b/Engine/Plugins/Animation/SkeletalMeshModelingTools/Source/SkeletalMeshModelingTools/Private/SkeletalMeshEditingCache.cpp
@@ -811,9 +811,9 @@ void USkeletalMeshEditingCache::DeformPreviewMesh(const TArray<FTransform>& Comp
 				
 	auto DeformPreviewMeshFunc = [&](FDynamicMesh3& VisualMesh)
 		{
-			auto WriteFunc = [&](int32 VertID, const FVector& PosedVertPos)
+			auto WriteFunc = [&](FVertInfo VertInfo, const FVector& PosedVertPos)
 				{
-					VisualMesh.SetVertex(VertID, PosedVertPos);
+					VisualMesh.SetVertex(VertInfo.VertID, PosedVertPos);
 				};
 					
 			GetPosedMesh(WriteFunc, EditingMesh, BoneMatrices, NAME_None, MorphTargetWeights);
diff --git a/Engine/Plugins/Animation/SkeletalMeshModelingTools/Source/SkeletalMeshModelingTools/Private/SkeletalMeshModelingToolsEditorMode.cpp b/Engine/Plugins/Animation/SkeletalMeshModelingTools/Source/SkeletalMeshModelingTools/Private/SkeletalMeshModelingToolsEditorMode.cpp
--- a/Engine/Plugins/Animation/SkeletalMeshModelingTools/Source/SkeletalMeshModelingTools/Private/SkeletalMeshModelingToolsEditorMode.cpp
+++ b/Engine/Plugins/Animation/SkeletalMeshModelingTools/Source/SkeletalMeshModelingTools/Private/SkeletalMeshModelingToolsEditorMode.cpp
@@ -12,6 +12,7 @@
 #include "Tools/EdModeInteractiveToolsContext.h"
 #include "EditorModeManager.h"
 #include "EditorViewportClient.h"
+#include "MeshAttributePaintToolV2.h"
 #include "Features/IModularFeatures.h"
 
 #include "SkeletalMeshModelingToolsModule.h"
@@ -358,7 +359,7 @@ void USkeletalMeshModelingToolsEditorMode::Enter()
 	RegisterTool(ToolManagerCommands.BeginDisplaceMeshTool, TEXT("BeginDisplaceMeshTool"), NewObject<UDisplaceMeshToolBuilder>());
 
 	RegisterTool(ToolManagerCommands.BeginAttributeEditorTool, TEXT("BeginAttributeEditorTool"), NewObject<UAttributeEditorToolBuilder>());
-	RegisterTool(ToolManagerCommands.BeginMeshAttributePaintTool, TEXT("BeginMeshAttributePaintTool"), NewObject<UMeshAttributePaintToolBuilder>());
+	RegisterTool(ToolManagerCommands.BeginMeshAttributePaintTool, TEXT("BeginMeshAttributePaintTool"), NewObject<UMeshAttributePaintToolV2Builder>());
 	RegisterTool(ToolManagerCommands.BeginMeshVertexPaintTool, TEXT("BeginMeshVertexPaintTool"), NewObject<UMeshVertexPaintToolBuilder>());
 	RegisterTool(ToolManagerCommands.BeginSkinWeightsPaintTool, TEXT("BeginSkinWeightsPaintTool"), NewObject<USkinWeightsPaintToolBuilder>());
 	RegisterTool(ToolManagerCommands.BeginSkinWeightsBindingTool, TEXT("BeginSkinWeightsBindingTool"), NewObject<USkinWeightsBindingToolBuilder>());
diff --git a/Engine/Plugins/Animation/SkeletalMeshModelingTools/Source/SkeletalMeshModelingTools/SkeletalMeshModelingTools.Build.cs b/Engine/Plugins/Animation/SkeletalMeshModelingTools/Source/SkeletalMeshModelingTools/SkeletalMeshModelingTools.Build.cs
--- a/Engine/Plugins/Animation/SkeletalMeshModelingTools/Source/SkeletalMeshModelingTools/SkeletalMeshModelingTools.Build.cs
+++ b/Engine/Plugins/Animation/SkeletalMeshModelingTools/Source/SkeletalMeshModelingTools/SkeletalMeshModelingTools.Build.cs
@@ -46,7 +46,8 @@ public class SkeletalMeshModelingTools : ModuleRules
 				"GeometryFramework",
 				"EditorWidgets",
 				"DeveloperSettings",
-				"DynamicMesh"
+				"DynamicMesh", 
+				"MeshModelingToolsEditorOnlyExp"
 			}
 		);
 
diff --git a/Engine/Plugins/Experimental/Animation/SkeletalMeshMorphTargetEditingTools/Source/SkeletalMeshMorphTargetEditingTools/Private/IMorphTargetEditingToolInterface.cpp b/Engine/Plugins/Experimental/Animation/SkeletalMeshMorphTargetEditingTools/Source/SkeletalMeshMorphTargetEditingTools/Private/IMorphTargetEditingToolInterface.cpp
--- a/Engine/Plugins/Experimental/Animation/SkeletalMeshMorphTargetEditingTools/Source/SkeletalMeshMorphTargetEditingTools/Private/IMorphTargetEditingToolInterface.cpp
+++ b/Engine/Plugins/Experimental/Animation/SkeletalMeshMorphTargetEditingTools/Source/SkeletalMeshMorphTargetEditingTools/Private/IMorphTargetEditingToolInterface.cpp
@@ -17,12 +17,8 @@ void IMorphTargetEditingToolInterface::SetupMorphEditingToolCommon()
 {
 	// Only support single selection tools at the moment
 	USingleSelectionTool* SingleTargetTool = CastChecked<USingleSelectionTool>(this);
-	ISkeletalMeshEditingInterface* EditingInterface = CastChecked<ISkeletalMeshEditingInterface>(this);
 	USkeletalMeshEditorContextObjectBase* EditorContext = SingleTargetTool->GetToolManager()->GetContextObjectStore()->FindContext<USkeletalMeshEditorContextObjectBase>();
-	if (EditorContext)
-	{
-		EditorContext->BindTo(EditingInterface);
-	}
+
 	
 	auto SetupFunction = [&](UMorphTargetEditingToolProperties* InProperties)
 	{
@@ -34,13 +30,5 @@ void IMorphTargetEditingToolInterface::SetupMorphEditingToolCommon()
 
 void IMorphTargetEditingToolInterface::ShutdownMorphEditingToolCommon()
 {
-	ISkeletalMeshEditingInterface* EditingInterface = CastChecked<ISkeletalMeshEditingInterface>(this);
 
-	// Only support single selection tools at the moment
-	USingleSelectionTool* SingleTargetTool = CastChecked<USingleSelectionTool>(this);
-	
-	if (USkeletalMeshEditorContextObjectBase* EditorContext = SingleTargetTool->GetToolManager()->GetContextObjectStore()->FindContext<USkeletalMeshEditorContextObjectBase>())
-	{
-		EditorContext->UnbindFrom(EditingInterface);
-	}	
 }
diff --git a/Engine/Plugins/Experimental/Animation/SkeletalMeshMorphTargetEditingTools/Source/SkeletalMeshMorphTargetEditingTools/Private/MorphTargetMaskTool.cpp b/Engine/Plugins/Experimental/Animation/SkeletalMeshMorphTargetEditingTools/Source/SkeletalMeshMorphTargetEditingTools/Private/MorphTargetMaskTool.cpp
new file mode 100644
--- /dev/null
+++ b/Engine/Plugins/Experimental/Animation/SkeletalMeshMorphTargetEditingTools/Source/SkeletalMeshMorphTargetEditingTools/Private/MorphTargetMaskTool.cpp
@@ -0,0 +1,722 @@
+// Copyright Epic Games, Inc. All Rights Reserved.
+
+#include "MorphTargetMaskTool.h"
+
+#include "ContextObjectStore.h"
+#include "Editor.h"
+#include "InteractiveToolManager.h"
+#include "ModelingToolTargetUtil.h"
+#include "MorphTargetVertexSculptTool.h"
+#include "PreviewMesh.h"
+#include "ReferenceSkeleton.h"
+#include "Drawing/MeshElementsVisualizer.h"
+#include "SkeletalMesh/SkeletalMeshEditingInterface.h"
+#include "TargetInterfaces/DynamicMeshCommitter.h"
+#include "TargetInterfaces/DynamicMeshProvider.h"
+#include "TargetInterfaces/MaterialProvider.h"
+#include "TargetInterfaces/SceneComponentBackedTarget.h"
+#include "TargetInterfaces/SkeletalMeshBackedTarget.h"
+#include "TargetInterfaces/SkeletonProvider.h"
+#include "Widgets/SWindow.h"
+#include "Widgets/Input/SButton.h"
+#include "Widgets/Input/SEditableTextBox.h"
+#include "Widgets/Text/STextBlock.h"
+
+
+#include UE_INLINE_GENERATED_CPP_BY_NAME(MorphTargetMaskTool)
+
+#define LOCTEXT_NAMESPACE "MorphTargetMaskTool"
+
+const TCHAR* UMorphTargetMaskToolProperties::CreateNewVertexAttributeOption = TEXT("[ Save new Mask ]");
+const TCHAR* UMorphTargetMaskToolProperties::DoNotSaveVertexAttributeOption = TEXT("[ Temporary Mask ]");
+
+bool UMorphTargetMaskToolBuilder::CanBuildTool(const FToolBuilderState& SceneState) const
+{
+	USkeletalMeshEditorContextObjectBase* EditorContext = SceneState.ToolManager->GetContextObjectStore()->FindContext<USkeletalMeshEditorContextObjectBase>();
+	if (!EditorContext)
+	{
+		return false;
+	}
+
+	if (EditorContext->GetEditingMorphTarget() == NAME_None)
+	{
+		return false;
+	}
+	
+	return Super::CanBuildTool(SceneState);
+}
+
+UMeshSurfacePointTool* UMorphTargetMaskToolBuilder::CreateNewTool(const FToolBuilderState& SceneState) const
+{
+	UMorphTargetMaskTool* MorphTargetEditorTool = NewObject<UMorphTargetMaskTool>(SceneState.ToolManager);
+	MorphTargetEditorTool->SetWorld(SceneState.World);
+	return MorphTargetEditorTool;
+}
+
+const FToolTargetTypeRequirements& UMorphTargetMaskToolBuilder::GetTargetRequirements() const
+{
+	static FToolTargetTypeRequirements ToolRequirements({
+		UMaterialProvider::StaticClass(),
+		UDynamicMeshProvider::StaticClass(),
+		UDynamicMeshCommitter::StaticClass(),
+		USceneComponentBackedTarget::StaticClass(),
+		USkeletonProvider::StaticClass(),
+	});
+
+	return ToolRequirements;
+}
+
+void UMorphTargetMaskToolProperties::Setup(UMorphTargetMaskTool* InTool, TArray<FName> InVertexAttributes)
+{
+	SetFlags(RF_Transactional);
+	
+	// Setup Tool Settings
+	FString AssetPath = TEXT("UnknownAsset");
+	UToolTarget* Target = InTool->GetTarget();
+	if (Cast<ISkeletalMeshBackedTarget>(Target))
+	{
+		if (USkeletalMesh* SkeletalMesh = Cast<ISkeletalMeshBackedTarget>(Target)->GetSkeletalMesh())
+		{
+			AssetPath = SkeletalMesh->GetPathName();
+		}
+	}
+	
+	PropertySetCacheIdentifier = FString::Printf(TEXT("%s::%s"), *AssetPath, *InTool->GetEditingMorphTarget().ToString());
+	RestoreProperties(InTool, PropertySetCacheIdentifier); 
+	
+	InTool->GetMorphTargetWeights().GenerateKeyArray(MorphTargets);
+	if (!MorphTargets.Contains(BaseMorphTarget))
+	{
+		BaseMorphTarget = InTool->GetEditingMorphTarget();
+	}
+
+	VertexAttributeOptions.Add(CreateNewVertexAttributeOption);
+	VertexAttributeOptions.Add(DoNotSaveVertexAttributeOption);
+	VertexAttributeOptions.Append(InVertexAttributes);
+	VertexAttributes = MoveTemp(InVertexAttributes);
+
+
+	if (VertexAttributeOption != CreateNewVertexAttributeOption && VertexAttributeOption != DoNotSaveVertexAttributeOption)
+	{
+		if (!VertexAttributes.Contains(VertexAttributeOption))
+		{
+			VertexAttributeOption = CreateNewVertexAttributeOption;
+		}
+	}
+}
+
+void UMorphTargetMaskToolProperties::Shutdown(UMorphTargetMaskTool* InTool , EToolShutdownType InShutdownType)
+{
+	if (InShutdownType != EToolShutdownType::Cancel)
+	{
+		SaveProperties(InTool, PropertySetCacheIdentifier);
+	}
+}
+
+bool UMorphTargetMaskToolProperties::ShouldSaveMask() const
+{
+	return VertexAttributeOption != DoNotSaveVertexAttributeOption;
+}
+
+bool UMorphTargetMaskToolProperties::ShouldSaveNewMask() const
+{
+	return VertexAttributeOption == CreateNewVertexAttributeOption;
+}
+
+const TArray<FName>& UMorphTargetMaskToolProperties::GetMorphTargets()
+{
+	return MorphTargets;
+}
+
+const TArray<FName>& UMorphTargetMaskToolProperties::GetVertexAttributeOptions()
+{
+	return VertexAttributeOptions;
+}
+
+TOptional<int32> UMorphTargetMaskToolProperties::GetVertexAttributeIndex()
+{
+	if (VertexAttributeOption != CreateNewVertexAttributeOption && 
+		VertexAttributeOption != DoNotSaveVertexAttributeOption)
+	{
+		return VertexAttributes.IndexOfByKey(VertexAttributeOption);
+	}
+
+	return {};
+}
+
+bool UMorphTargetMaskTool::SetupToolMesh(FDynamicMesh3& InOutToolMesh, int32& OutInitialAttributeIndex)
+{
+	using namespace UE::Geometry;
+
+	TArray<FName> VertexAttributes;
+	for (int32 WeightMapIndex = 0; WeightMapIndex < InOutToolMesh.Attributes()->NumWeightLayers(); WeightMapIndex++)
+	{
+		FName AttributeName = InOutToolMesh.Attributes()->GetWeightLayer(WeightMapIndex)->GetName();
+		VertexAttributes.Add(AttributeName);
+	}
+	
+	MaskToolProperties->Setup(this, MoveTemp(VertexAttributes));
+	
+	MaskToolCommitMesh = InOutToolMesh;
+	{
+		// Create a temp copy of the existing morph target so that a weight map can be applied to it to create the new replacement morph target 
+		FDynamicMeshMorphTargetAttribute* MorphTargetAttribute = MaskToolCommitMesh.Attributes()->GetMorphTargetAttribute(ToolMorphTargetName);
+
+		ToolMorphTargetBackupAttributeName = TEXT("MorphTargetMaskTool_ToolMorphTargetBackup");
+
+		while (MaskToolCommitMesh.Attributes()->HasMorphTargetAttribute(ToolMorphTargetBackupAttributeName))
+		{
+			ToolMorphTargetBackupAttributeName.SetNumber(ToolMorphTargetBackupAttributeName.GetNumber() + 1);
+		}
+	
+		FDynamicMeshMorphTargetAttribute* ToolMorphTargetBackupAttribute = new FDynamicMeshMorphTargetAttribute(&MaskToolCommitMesh);
+	
+		ToolMorphTargetBackupAttribute->Copy(*MorphTargetAttribute);
+		ToolMorphTargetBackupAttribute->SetName(ToolMorphTargetBackupAttributeName);
+			
+		MaskToolCommitMesh.Attributes()->AttachMorphTargetAttribute(ToolMorphTargetBackupAttributeName, ToolMorphTargetBackupAttribute);	
+	}
+	
+	// Create a temp weight layer to paint into
+	const int32 NumAttributeLayers = InOutToolMesh.Attributes()->NumWeightLayers();
+	InOutToolMesh.Attributes()->SetNumWeightLayers(NumAttributeLayers + 1);
+	FDynamicMeshWeightAttribute* ActiveWeightMap = InOutToolMesh.Attributes()->GetWeightLayer(NumAttributeLayers);
+	ActiveWeightMap->SetName(FName("PaintLayer"));
+
+	TempPaintLayerAttributeIndex = NumAttributeLayers;
+
+	
+	if (TOptional<int32> AttributeIndex = MaskToolProperties->GetVertexAttributeIndex())
+	{
+		OutInitialAttributeIndex = *AttributeIndex;
+	}
+	else
+	{
+		OutInitialAttributeIndex = TempPaintLayerAttributeIndex;
+	}
+
+	CurrentAttributeIndex = OutInitialAttributeIndex;
+	
+	return true;
+}
+
+
+
+namespace UE::MorphTargetMaskTool::Private
+{
+	class SSaveMaskWindow : public SCompoundWidget
+	{
+	public:
+		struct FResult
+		{
+			FText AttributeName;
+			bool bSave = false;
+		};
+		
+		SLATE_BEGIN_ARGS(SSaveMaskWindow) {}
+			SLATE_ARGUMENT( TSharedPtr<SWindow>, WidgetWindow )
+			SLATE_ARGUMENT( TArray<FName>, ExistingNames )
+			SLATE_ARGUMENT( FResult*, ResultPtr)
+		SLATE_END_ARGS()
+
+		virtual ~SSaveMaskWindow() = default;
+		
+		void Construct(const FArguments& InArgs)
+		{
+			WidgetWindow = InArgs._WidgetWindow;
+			ExistingNames = InArgs._ExistingNames;
+			ResultPtr = InArgs._ResultPtr;
+
+			this->ChildSlot
+			[
+				SNew(SVerticalBox)
+				+SVerticalBox::Slot()
+				.FillHeight(1.0f)
+				[
+					SNew(SHorizontalBox)
+					+ SHorizontalBox::Slot()
+					.AutoWidth()
+					.VAlign(VAlign_Center)
+					.Padding(FMargin(6.f, 0.0))
+					[
+						SNew(STextBlock)
+						.Text(FText::FromString("Mask Name: "))
+					]
+					+ SHorizontalBox::Slot()
+					.AutoWidth()
+					.VAlign(VAlign_Center)
+					.Padding(FMargin(6.f, 0.0))
+					[
+						SNew(SEditableTextBox)
+						.Text_Lambda([this]()
+							{
+								return ResultPtr->AttributeName;
+							})
+						.OnTextChanged_Lambda([this](const FText& InText)
+							{
+								ResultPtr->AttributeName = InText;
+							})
+						.OnVerifyTextChanged_Lambda([this](const FText&, FText& OutErrorMessage)
+							{
+								FName Name = *ResultPtr->AttributeName.ToString();
+								if (ExistingNames.Contains(Name))
+								{
+									OutErrorMessage = FText::FromString(TEXT("The name is already used by another attribute"));
+									return false;
+								}
+
+								return true;
+							})
+						.OnTextCommitted_Lambda([this](const FText&, ETextCommit::Type CommitType)
+							{
+								if (CommitType == ETextCommit::OnEnter)
+								{
+									OnSave();
+								}
+							})
+					]
+				]
+				+SVerticalBox::Slot()
+				.AutoHeight()
+				[
+					SNew(SHorizontalBox)
+					+SHorizontalBox::Slot()
+					[
+						SNew(SButton)
+						.HAlign(HAlign_Center)
+						.Text(LOCTEXT("SSaveMaskWindow_SaveButton", "Save"))
+						.OnClicked(this, &SSaveMaskWindow::OnSave)
+					]
+					+SHorizontalBox::Slot()
+					[
+						SNew(SButton)
+						.HAlign(HAlign_Center)
+						.Text(LOCTEXT("SSaveMaskWindow_SkipSavingButton", "Skip Saving"))
+						.OnClicked(this, &SSaveMaskWindow::OnSkip)
+					]
+				]
+			];
+		}
+
+		bool SupportsKeyboardFocus() const override { return true; }
+
+		bool IsNameValid()
+		{
+			return !ExistingNames.Contains(FName(*ResultPtr->AttributeName.ToString()));
+		}
+		FReply OnSave()
+		{
+			if (IsNameValid())
+			{
+				ResultPtr->bSave = true;
+				CloseWindow();
+			}
+			
+			return FReply::Handled();
+		}
+
+		FReply OnSkip()
+		{
+			CloseWindow();
+			return FReply::Handled();
+		}
+
+		void CloseWindow()
+		{
+			if ( WidgetWindow.IsValid() )
+			{
+				WidgetWindow.Pin()->RequestDestroyWindow();
+			}	
+		}
+
+		TWeakPtr<SWindow> WidgetWindow;
+		TArray<FName> ExistingNames;
+
+		FResult* ResultPtr = nullptr;
+	};
+	
+
+}
+
+void UMorphTargetMaskTool::CommitToolMesh(FDynamicMesh3& InToolMesh)
+{
+	using namespace UE::Geometry;
+
+	if (MaskToolProperties->ShouldSaveMask())
+	{
+		FDynamicMeshWeightAttribute* TargetWeightMap = nullptr;
+		if (!MaskToolProperties->ShouldSaveNewMask())
+		{
+			TargetWeightMap = MaskToolCommitMesh.Attributes()->GetWeightLayer(CurrentAttributeIndex);
+		}
+		else
+		{
+			/** Create the window to host our package dialog widget */
+			TSharedRef<SWindow> SaveMaskModalWindow = SNew(SWindow)
+				.Title(FText::FromString("Save Vertex Attribute"))
+				.SizingRule(ESizingRule::Autosized);
+
+			
+			using namespace UE::MorphTargetMaskTool::Private;
+			SSaveMaskWindow::FResult Result;
+			FName InitialName = FName(GetBaseMorphTarget().ToString() + TEXT("_Mask"));
+
+			while (MaskToolProperties->VertexAttributes.Contains(InitialName))
+			{
+				InitialName.SetNumber(InitialName.GetNumber() + 1);
+			}
+			
+			Result.AttributeName = FText::FromName(InitialName);
+			
+			/** Set the content of the window to our package dialog widget */
+			TSharedRef<SSaveMaskWindow> SaveMaskWindowWidget=
+				SNew(SSaveMaskWindow)
+				.WidgetWindow(SaveMaskModalWindow)
+				.ExistingNames(MaskToolProperties->VertexAttributes)
+				.ResultPtr(&Result);
+
+			SaveMaskModalWindow->SetContent(SaveMaskWindowWidget);
+
+			/** Show the dialog window as a modal window */
+			GEditor->EditorAddModalWindow(SaveMaskModalWindow);
+
+			if (Result.bSave)
+			{
+				// Create a new weight layer to save into
+				const int32 NumAttributeLayers = MaskToolCommitMesh.Attributes()->NumWeightLayers();
+				MaskToolCommitMesh.Attributes()->SetNumWeightLayers(NumAttributeLayers + 1);
+				TargetWeightMap = MaskToolCommitMesh.Attributes()->GetWeightLayer(NumAttributeLayers);
+				TargetWeightMap->SetName(*Result.AttributeName.ToString());
+
+				// Override the chosen option now that we know the name of the new attribute
+				MaskToolProperties->VertexAttributeOption = TargetWeightMap->GetName();
+			}
+		}
+
+		if (TargetWeightMap)
+		{
+			FDynamicMeshWeightAttribute* SourceWeightMap = InToolMesh.Attributes()->GetWeightLayer(CurrentAttributeIndex);
+	
+			ParallelFor(MaskToolCommitMesh.MaxVertexID(), [this, SourceWeightMap, TargetWeightMap](int32 VertID)
+			{
+				if (!MaskToolCommitMesh.IsVertex(VertID))
+				{
+					return;
+				}
+				
+				float Weight = 0.0f;
+				SourceWeightMap->GetValue(VertID, &Weight);
+				TargetWeightMap->SetValue(VertID, &Weight);
+			});	
+		}
+	}
+
+	MaskToolCommitMesh.Attributes()->RemoveMorphTargetAttribute(ToolMorphTargetBackupAttributeName);
+	
+	if (UInteractiveToolManager* ToolManager = GetToolManager())
+	{
+		ToolManager->BeginUndoTransaction(LOCTEXT("MorphTargetMaskTool_TransactionName", "Mask Morph Target"));
+
+		// commit the mask tool mesh cache instead of the sculpt mesh/base tool tool mesh
+		// setting ModifiedTopology to true to force a full dyna mesh update such that weight layer changes
+		// are converted to mesh description as well;
+		// todo: update weight layers only
+		constexpr bool bForceModifiedTopology = true;
+		UE::ToolTarget::CommitDynamicMeshUpdate(Target, MaskToolCommitMesh, bForceModifiedTopology);
+		
+		ToolManager->EndUndoTransaction();
+	}
+
+	if (EditorContext.IsValid())
+	{
+		EditorContext->NotifyMorphTargetEdited();
+	}
+}
+
+void UMorphTargetMaskTool::UpdatePreview(const TSet<int32>* TrianglesToUpdate, const TArray<int32>* VerticesToUpdate)
+{
+	Super::UpdatePreview(TrianglesToUpdate, VerticesToUpdate);
+
+	using namespace UE::Geometry;
+	
+	FDynamicMeshMorphTargetAttribute* ToolMorphTargetCache = MaskToolCommitMesh.Attributes()->GetMorphTargetAttribute(ToolMorphTargetName);
+
+	FName BaseMorphTargetAttributeName = GetBaseMorphTarget() == ToolMorphTargetName ? ToolMorphTargetBackupAttributeName : GetBaseMorphTarget();  
+	FDynamicMeshMorphTargetAttribute* BaseToolMorphTarget = MaskToolCommitMesh.Attributes()->GetMorphTargetAttribute(BaseMorphTargetAttributeName);
+
+	int32 NumVerts = VerticesToUpdate ? VerticesToUpdate->Num() : MaskToolCommitMesh.MaxVertexID();
+	
+	ParallelFor(NumVerts, [VerticesToUpdate, this, ToolMorphTargetCache, BaseToolMorphTarget](int32 ThreadID)
+	{
+		int32 VertexID = VerticesToUpdate
+			? (*VerticesToUpdate)[ThreadID]
+			: ThreadID;
+
+		if (!MaskToolCommitMesh.IsVertex(VertexID))
+		{
+			return;
+		}
+
+		const float Weight = VertexData.GetValue(VertexID);
+
+		FVector Delta;
+		BaseToolMorphTarget->GetValue(VertexID, Delta);
+
+		Delta *= (1.0f - Weight);
+		ToolMorphTargetCache->SetValue(VertexID, Delta);
+	});
+
+	const TArray<FTransform>& ComponentSpaceTransforms = GetComponentSpaceBoneTransforms();
+	const TMap<FName, float>& MorphTargetWeights = GetMorphTargetWeights();
+	PosePreviewMesh(ComponentSpaceTransforms, MorphTargetWeights, TrianglesToUpdate, VerticesToUpdate);
+	
+	if (TrianglesToUpdate)
+	{
+		Octree.ReinsertTriangles(*TrianglesToUpdate);
+	}
+	else
+	{
+		bRebuildOctree = true;
+	}
+
+	bMeshSelectorNeedsUpdate = true;
+}
+
+FName UMorphTargetMaskTool::GetEditingMorphTarget()
+{
+	if (EditorContext.IsValid())
+	{
+		return EditorContext->GetEditingMorphTarget();
+	}
+
+	return {};	
+}
+
+FName UMorphTargetMaskTool::GetBaseMorphTarget()
+{
+	return MaskToolProperties->BaseMorphTarget;
+}
+
+TMap<FName, float> UMorphTargetMaskTool::GetMorphTargetWeights()
+{
+	if (IsBrushMirroringEnabled())
+	{
+		return {};
+	}
+	
+	if (EditorContext.IsValid())
+	{
+		return EditorContext->GetMorphTargetWeights();
+	}
+
+	return {};	
+}
+
+const TArray<FTransform>& UMorphTargetMaskTool::GetComponentSpaceBoneTransforms()
+{
+	if (IsBrushMirroringEnabled())
+	{
+		return ComponentSpaceTransformsRefPose;
+	}
+	
+	if (EditorContext.IsValid())
+	{
+		return EditorContext->GetComponentSpaceBoneTransforms(GetTarget());
+	}
+
+	if (DefaultComponentSpaceBoneTransforms.IsEmpty())
+	{
+		ISkeletonProvider* SkeletonProvider = CastChecked<ISkeletonProvider>(GetTarget());
+		SkeletonProvider->GetSkeleton().GetBoneAbsoluteTransforms(DefaultComponentSpaceBoneTransforms);
+	}
+	
+	return DefaultComponentSpaceBoneTransforms;
+}
+
+void UMorphTargetMaskTool::ToggleBoneManipulation(bool bEnable)
+{
+	if (EditorContext.IsValid())
+	{
+		EditorContext->ToggleBoneManipulation(bEnable);
+	}
+}
+
+void UMorphTargetMaskTool::HandlePoseChangeDetectorEvent(SkeletalMeshToolsHelper::FPoseChangeDetector::FPayload Payload)
+{
+	if (InStroke())
+	{
+		return;
+	}
+	
+	using namespace SkeletalMeshToolsHelper;
+	if (Payload.CurrentState == FPoseChangeDetector::PoseStoppedChanging)
+	{
+		bRebuildOctree = true;
+	}
+	else if (Payload.CurrentState == FPoseChangeDetector::PoseJustChanged ||
+		Payload.CurrentState == FPoseChangeDetector::PoseChanged)
+	{
+		bFastDeformPreviewMesh = true;
+		bMeshSelectorNeedsUpdate = true;
+	}
+}
+
+
+
+void UMorphTargetMaskTool::Setup()
+{
+	EditorContext = GetToolManager()->GetContextObjectStore()->FindContext<USkeletalMeshEditorContextObjectBase>();
+	SetupMorphEditingToolCommon();
+	
+	ToolMorphTargetName = GetEditingMorphTarget();
+
+
+
+	// Add the mask tool settings to the top. Property init takes place during SetupToolMesh()
+	MaskToolProperties = NewObject<UMorphTargetMaskToolProperties>(this);
+	AddToolPropertySource(MaskToolProperties);
+	
+
+
+	
+	const FReferenceSkeleton& RefSkeleton = CastChecked<ISkeletonProvider>(GetTarget())->GetSkeleton();
+	RefSkeleton.GetBoneAbsoluteTransforms(ComponentSpaceTransformsRefPose);
+	
+	MorphTargetZeroWeights = GetMorphTargetWeights();
+	for (TPair<FName, float>& MorphInfo : MorphTargetZeroWeights)
+	{
+		MorphInfo.Value = 0.0f;
+	}
+
+	ToggleBoneManipulation(true);
+
+	PoseChangeDetector.GetNotifier().AddUObject(this, &UMorphTargetMaskTool::HandlePoseChangeDetectorEvent);
+	
+	// Setup Vertex Sculpt Tool
+	Super::Setup();
+
+
+	BaseMorphTargetWatcher.Initialize([this]()
+	{
+		return GetBaseMorphTarget();
+	},
+[this](const FName& InNewBaseMorphTarget)
+	{
+		UpdatePreview();
+	},
+GetBaseMorphTarget());
+
+	
+	VertexAttributeOptionWatcher.Initialize([this]()
+	{
+		return MaskToolProperties->VertexAttributeOption;
+	},
+[this](const FName& InNewVertexAttributeOption)
+	{
+		if (InNewVertexAttributeOption != MaskToolProperties->CreateNewVertexAttributeOption && 
+			InNewVertexAttributeOption != MaskToolProperties->DoNotSaveVertexAttributeOption)
+		{
+			// Editing a specific attribute
+			int32 WeightMapIndex = MaskToolProperties->VertexAttributes.IndexOfByKey(InNewVertexAttributeOption);
+			SetAttributeToPaint(WeightMapIndex);
+			CurrentAttributeIndex = WeightMapIndex;
+		}
+		else
+		{
+			if (CurrentAttributeIndex != TempPaintLayerAttributeIndex)
+			{
+				SetAttributeToPaint(TempPaintLayerAttributeIndex);
+				CurrentAttributeIndex = TempPaintLayerAttributeIndex;
+			}
+		}
+	},
+	MaskToolProperties->VertexAttributeOption);
+}
+
+
+
+void UMorphTargetMaskTool::Shutdown(EToolShutdownType ShutdownType)
+{
+	Super::Shutdown(ShutdownType);
+
+	MaskToolProperties->Shutdown(this, ShutdownType);
+	
+	ShutdownMorphEditingToolCommon();	
+}
+
+void UMorphTargetMaskTool::OnTick(float DeltaTime)
+{
+	Super::OnTick(DeltaTime);
+
+	if (!InStroke())
+	{
+		BaseMorphTargetWatcher.CheckAndUpdate();
+		VertexAttributeOptionWatcher.CheckAndUpdate();
+		
+		const TArray<FTransform>& ComponentSpaceTransforms = GetComponentSpaceBoneTransforms();
+		TMap<FName, float> MorphTargetWeights = GetMorphTargetWeights();
+		
+		PoseChangeDetector.CheckPose(ComponentSpaceTransforms, MorphTargetWeights);
+
+		if (bFastDeformPreviewMesh)
+		{
+			bFastDeformPreviewMesh = false;
+			PosePreviewMesh(ComponentSpaceTransforms, MorphTargetWeights);
+		}
+
+		if (GetEditingMode() == EMeshVertexAttributePaintToolEditMode::Mesh && bMeshSelectorNeedsUpdate)
+		{
+			bMeshSelectorNeedsUpdate = false;
+			MeshSelector->UpdateAfterMeshDeformation();	
+		}
+
+		if (bRebuildOctree)
+		{
+			bRebuildOctree = false;
+
+			RebuildOctree();
+		}
+	}
+}
+
+void UMorphTargetMaskTool::SetupCommonProperties(const TFunction<void(UMorphTargetEditingToolProperties*)>& InSetupFunction)
+{
+	EditorToolProperties = NewObject<UMorphTargetEditingToolProperties>(this);
+
+	InSetupFunction(EditorToolProperties);
+	
+	AddToolPropertySource(EditorToolProperties);
+}
+
+void UMorphTargetMaskTool::PosePreviewMesh(const TArray<FTransform>& ComponentSpaceTransforms, const TMap<FName, float>& MorphTargetWeights, const TSet<int32>* TrianglesToUpdate, const TArray<int32>* VerticesToUpdate)
+{
+	using namespace SkeletalMeshToolsHelper;
+	
+	TArray<FMatrix> BoneMatrices = ComputeBoneMatrices(ComponentSpaceTransformsRefPose, ComponentSpaceTransforms);
+
+	FDynamicMesh3& SculptMesh = *GetSculptMesh();
+
+	auto WriteFunc = [&](FVertInfo VertInfo, const FVector& PosedVertPos)
+		{
+			SculptMesh.SetVertex(VertInfo.VertID, PosedVertPos);
+		};
+
+	GetPosedMesh(WriteFunc, MaskToolCommitMesh, BoneMatrices, NAME_None, MorphTargetWeights, VerticesToUpdate ? *VerticesToUpdate : TArray<int32>());
+	if (!TrianglesToUpdate)
+	{
+		FMeshNormals::QuickRecomputeOverlayNormals(SculptMesh);
+	}
+	else
+	{
+		FMeshNormals::RecomputeOverlayTriNormals(SculptMesh, TrianglesToUpdate->Array());
+	}
+
+	constexpr bool bNormal = true;
+	CastChecked<UDynamicMeshComponent>(GetSculptMeshComponent())->FastNotifyPositionsUpdated(bNormal);
+
+	MeshElementsDisplay->NotifyMeshChanged();
+}
+
+
+#undef LOCTEXT_NAMESPACE
+
diff --git a/Engine/Plugins/Experimental/Animation/SkeletalMeshMorphTargetEditingTools/Source/SkeletalMeshMorphTargetEditingTools/Private/MorphTargetVertexSculptTool.cpp b/Engine/Plugins/Experimental/Animation/SkeletalMeshMorphTargetEditingTools/Source/SkeletalMeshMorphTargetEditingTools/Private/MorphTargetVertexSculptTool.cpp
--- a/Engine/Plugins/Experimental/Animation/SkeletalMeshMorphTargetEditingTools/Source/SkeletalMeshMorphTargetEditingTools/Private/MorphTargetVertexSculptTool.cpp
+++ b/Engine/Plugins/Experimental/Animation/SkeletalMeshMorphTargetEditingTools/Source/SkeletalMeshMorphTargetEditingTools/Private/MorphTargetVertexSculptTool.cpp
@@ -428,14 +428,14 @@ void UMorphTargetVertexSculptTool::PoseSculptMesh(const TArray<FTransform>& Comp
 
 			float EditingMorphTargetWeight = MorphTargetWeights[ToolMorphTargetName];
 	
-			auto WriteFunc = [&](int32 VertID, const FVector& PosedVertPos)
+			auto WriteFunc = [&](FVertInfo VertInfo, const FVector& PosedVertPos)
 				{
-					SculptMesh.SetVertex(VertID, PosedVertPos);
+					SculptMesh.SetVertex(VertInfo.VertID, PosedVertPos);
 
 					FVector Delta;
-					MorphTargetAttribute->GetValue(VertID, Delta);
+					MorphTargetAttribute->GetValue(VertInfo.VertID, Delta);
 			
-					MeshWithoutEditingMorph.SetVertex(VertID,  PosedVertPos - (Delta * EditingMorphTargetWeight));
+					MeshWithoutEditingMorph.SetVertex(VertInfo.VertID,  PosedVertPos - (Delta * EditingMorphTargetWeight));
 				};
 
 			GetPosedMesh(WriteFunc, InMesh, BoneMatrices, NAME_None, MorphTargetWeights);
diff --git a/Engine/Plugins/Experimental/Animation/SkeletalMeshMorphTargetEditingTools/Source/SkeletalMeshMorphTargetEditingTools/Private/SKMMorphTargetEditingToolsCommands.cpp b/Engine/Plugins/Experimental/Animation/SkeletalMeshMorphTargetEditingTools/Source/SkeletalMeshMorphTargetEditingTools/Private/SKMMorphTargetEditingToolsCommands.cpp
--- a/Engine/Plugins/Experimental/Animation/SkeletalMeshMorphTargetEditingTools/Source/SkeletalMeshMorphTargetEditingTools/Private/SKMMorphTargetEditingToolsCommands.cpp
+++ b/Engine/Plugins/Experimental/Animation/SkeletalMeshMorphTargetEditingTools/Source/SkeletalMeshMorphTargetEditingTools/Private/SKMMorphTargetEditingToolsCommands.cpp
@@ -22,6 +22,7 @@ void FSkeletalMeshMorphTargetEditingToolsCommands::RegisterCommands()
 	UI_COMMAND(BeginMorphTargetTool, "Morph", "Edit Morph Targets", EUserInterfaceActionType::ToggleButton, FInputChord());
 	
 	UI_COMMAND(BeginMorphTargetSculptTool, "Sculpt Morph Target", "To use this tool, select a Morph Target for editing from the morph target tab using the radio button", EUserInterfaceActionType::ToggleButton, FInputChord());
+	UI_COMMAND(BeginMorphTargetMaskTool, "Mask Morph Target", "To use this tool, select a Morph Target for editing from the morph target tab using the radio button", EUserInterfaceActionType::ToggleButton, FInputChord());
 }
 
 #undef LOCTEXT_NAMESPACE
diff --git a/Engine/Plugins/Experimental/Animation/SkeletalMeshMorphTargetEditingTools/Source/SkeletalMeshMorphTargetEditingTools/Private/SKMMorphTargetEditingToolsModule.cpp b/Engine/Plugins/Experimental/Animation/SkeletalMeshMorphTargetEditingTools/Source/SkeletalMeshMorphTargetEditingTools/Private/SKMMorphTargetEditingToolsModule.cpp
--- a/Engine/Plugins/Experimental/Animation/SkeletalMeshMorphTargetEditingTools/Source/SkeletalMeshMorphTargetEditingTools/Private/SKMMorphTargetEditingToolsModule.cpp
+++ b/Engine/Plugins/Experimental/Animation/SkeletalMeshMorphTargetEditingTools/Source/SkeletalMeshMorphTargetEditingTools/Private/SKMMorphTargetEditingToolsModule.cpp
@@ -2,6 +2,7 @@
 
 #include "SKMMorphTargetEditingToolsModule.h"
 
+#include "MorphTargetMaskTool.h"
 #include "MorphTargetVertexSculptTool.h"
 #include "SKMMorphTargetEditingToolsCommands.h"
 #include "SKMMorphTargetEditingToolsStyle.h"
@@ -38,6 +39,14 @@ void FSkeletalMeshMorphTargetEditingToolsModule::GetExtensionTools(const FExtens
 		ToolDesc.ToolBuilder = NewObject<UMorphTargetVertexSculptToolBuilder>();
 		OutTools.Add(ToolDesc);
 	}
+	
+	{
+		FExtensionToolDescription ToolDesc;
+		ToolDesc.ToolName = LOCTEXT("SkeletalMeshMorphTargetMaskTool", "Mask Morph Target");
+		ToolDesc.ToolCommand = Commands.BeginMorphTargetMaskTool;
+		ToolDesc.ToolBuilder = NewObject<UMorphTargetMaskToolBuilder>();
+		OutTools.Add(ToolDesc);
+	}
 }
 
 bool FSkeletalMeshMorphTargetEditingToolsModule::GetExtensionExtendedInfo(FModelingModeExtensionExtendedInfo& InfoOut)
diff --git a/Engine/Plugins/Experimental/Animation/SkeletalMeshMorphTargetEditingTools/Source/SkeletalMeshMorphTargetEditingTools/Public/MorphTargetEditingToolProperties.h b/Engine/Plugins/Experimental/Animation/SkeletalMeshMorphTargetEditingTools/Source/SkeletalMeshMorphTargetEditingTools/Public/MorphTargetEditingToolProperties.h
--- a/Engine/Plugins/Experimental/Animation/SkeletalMeshMorphTargetEditingTools/Source/SkeletalMeshMorphTargetEditingTools/Public/MorphTargetEditingToolProperties.h
+++ b/Engine/Plugins/Experimental/Animation/SkeletalMeshMorphTargetEditingTools/Source/SkeletalMeshMorphTargetEditingTools/Public/MorphTargetEditingToolProperties.h
@@ -11,6 +11,6 @@ class UMorphTargetEditingToolProperties : public UInteractiveToolPropertySet
 {
 	GENERATED_BODY()
 public:
-	UPROPERTY(VisibleAnywhere, Category = "MorphTarget", meta = (GetOptions = GetMorphTargetNames, NoResetToDefault, DisplayName = "Name"))
+	UPROPERTY(VisibleAnywhere, Category = "MorphTarget", meta = (NoResetToDefault, DisplayName = "Name"))
 	FName EditMorphTargetName;
 };
\ No newline at end of file
diff --git a/Engine/Plugins/Experimental/Animation/SkeletalMeshMorphTargetEditingTools/Source/SkeletalMeshMorphTargetEditingTools/Public/MorphTargetMaskTool.h b/Engine/Plugins/Experimental/Animation/SkeletalMeshMorphTargetEditingTools/Source/SkeletalMeshMorphTargetEditingTools/Public/MorphTargetMaskTool.h
new file mode 100644
--- /dev/null
+++ b/Engine/Plugins/Experimental/Animation/SkeletalMeshMorphTargetEditingTools/Source/SkeletalMeshMorphTargetEditingTools/Public/MorphTargetMaskTool.h
@@ -0,0 +1,146 @@
+// Copyright Epic Games, Inc. All Rights Reserved.
+
+#pragma once
+
+#include "IMorphTargetEditingToolInterface.h"
+#include "MorphTargetEditingToolProperties.h"
+
+#include "MeshVertexAttributePaintToolBase.h"
+#include "SkeletalMesh/SkeletalMeshToolsHelper.h"
+#include "SkeletalMesh/SkeletalMeshEditingInterface.h"
+
+#include "MorphTargetMaskTool.generated.h"
+
+/**
+ * MorphTarget Vertex Sculpt Tool Builder
+ */
+UCLASS()
+class UMorphTargetMaskToolBuilder : public UMeshSurfacePointMeshEditingToolBuilder
+{
+	GENERATED_BODY()
+
+public:
+	virtual bool CanBuildTool(const FToolBuilderState& SceneState) const override;
+	virtual UMeshSurfacePointTool* CreateNewTool(const FToolBuilderState& SceneState) const override;
+	virtual const FToolTargetTypeRequirements& GetTargetRequirements() const override;
+};
+
+
+UCLASS(MinimalAPI)
+class UMorphTargetMaskToolProperties : public UInteractiveToolPropertySet
+{
+	GENERATED_BODY()
+public:
+
+	void Setup(UMorphTargetMaskTool* InTool, TArray<FName> InVertexAttributes);
+	void Shutdown(UMorphTargetMaskTool* InTool , EToolShutdownType InShutdownType);
+	bool ShouldSaveMask() const;
+	bool ShouldSaveNewMask() const;
+	
+	UPROPERTY(EditAnywhere, Category = "Mask Settings", meta = (GetOptions = GetMorphTargets,  NoResetToDefault))
+	FName BaseMorphTarget = NAME_None;
+
+	UFUNCTION()
+	const TArray<FName>& GetMorphTargets();
+
+	TArray<FName> MorphTargets;
+
+	UPROPERTY(EditAnywhere, Category = "Mask Settings", meta = (DisplayName = "Vertex Attribute", GetOptions = GetVertexAttributeOptions,  NoResetToDefault))
+	FName VertexAttributeOption = CreateNewVertexAttributeOption;
+	
+	UFUNCTION()
+	const TArray<FName>& GetVertexAttributeOptions();
+
+	TArray<FName> VertexAttributeOptions;
+	TArray<FName> VertexAttributes;
+
+	FString PropertySetCacheIdentifier;	
+
+	static const TCHAR* CreateNewVertexAttributeOption;
+	static const TCHAR* DoNotSaveVertexAttributeOption;
+
+	TOptional<int32> GetVertexAttributeIndex();
+};
+
+/**
+* MorphTarget Sculpt Tool
+*/
+UCLASS()
+class  UMorphTargetMaskTool: 
+	public UMeshVertexAttributePaintToolBase,
+	public IMorphTargetEditingToolInterface
+{
+	GENERATED_BODY()
+
+	virtual bool SetupToolMesh(FDynamicMesh3& InOutToolMesh, int32& OutInitialAttributeIndex) override;
+	virtual void CommitToolMesh(FDynamicMesh3& InToolMesh) override;
+	virtual void UpdatePreview(const TSet<int32>* TrianglesToUpdate = nullptr, const TArray<int32>* VerticesToUpdate = nullptr) override;
+public:
+	FName GetEditingMorphTarget();
+	FName GetBaseMorphTarget();
+	virtual TMap<FName, float> GetMorphTargetWeights();
+	virtual const TArray<FTransform>& GetComponentSpaceBoneTransforms();
+	virtual void ToggleBoneManipulation(bool bEnable);
+
+
+	// UMeshVertexSculptTool overrides
+	virtual void Setup() override;
+	virtual void Shutdown(EToolShutdownType ShutdownType) override;
+	virtual void OnTick(float DeltaTime) override;
+
+	
+protected:
+
+
+	void HandlePoseChangeDetectorEvent(SkeletalMeshToolsHelper::FPoseChangeDetector::FPayload Payload);
+
+	// IMorphTargetEditingTool overrides
+	virtual void SetupCommonProperties(const TFunction<void(UMorphTargetEditingToolProperties*)>& InSetupFunction) override;
+	
+	void PosePreviewMesh(
+		const TArray<FTransform>& ComponentSpaceTransforms,
+		const TMap<FName, float>& MorphTargetWeights,
+		const TSet<int32>* TrianglesToUpdate = nullptr,
+		const TArray<int32>* VerticesToUpdate = nullptr
+		);
+
+	UPROPERTY()
+	TObjectPtr<class UMorphTargetEditingToolProperties> EditorToolProperties;
+
+	UPROPERTY()
+	TObjectPtr<class UMorphTargetMaskToolProperties> MaskToolProperties;
+
+
+	
+	UPROPERTY()
+	TWeakObjectPtr<USkeletalMeshEditorContextObjectBase> EditorContext = nullptr;
+
+	SkeletalMeshToolsHelper::FPoseChangeDetector PoseChangeDetector;
+
+	FDynamicMesh3 MaskToolCommitMesh;
+
+	TArray<FTransform> ComponentSpaceTransformsRefPose;
+	TMap<FName, float> MorphTargetZeroWeights;
+	
+	FName ToolMorphTargetName;
+	FName ToolMorphTargetBackupAttributeName;
+
+	TArray<FTransform> DefaultComponentSpaceBoneTransforms;
+
+	TValueWatcher<FName> BaseMorphTargetWatcher;
+	TValueWatcher<FName> VertexAttributeOptionWatcher;
+
+	bool bFastDeformPreviewMesh = false;
+	bool bRebuildOctree = false;
+
+	bool bMeshSelectorNeedsUpdate = false;
+
+	
+	int32 TempPaintLayerAttributeIndex = INDEX_NONE;
+	int32 CurrentAttributeIndex = INDEX_NONE;
+};
+
+
+
+
+
diff --git a/Engine/Plugins/Experimental/Animation/SkeletalMeshMorphTargetEditingTools/Source/SkeletalMeshMorphTargetEditingTools/Public/SKMMorphTargetEditingToolsCommands.h b/Engine/Plugins/Experimental/Animation/SkeletalMeshMorphTargetEditingTools/Source/SkeletalMeshMorphTargetEditingTools/Public/SKMMorphTargetEditingToolsCommands.h
--- a/Engine/Plugins/Experimental/Animation/SkeletalMeshMorphTargetEditingTools/Source/SkeletalMeshMorphTargetEditingTools/Public/SKMMorphTargetEditingToolsCommands.h
+++ b/Engine/Plugins/Experimental/Animation/SkeletalMeshMorphTargetEditingTools/Source/SkeletalMeshMorphTargetEditingTools/Public/SKMMorphTargetEditingToolsCommands.h
@@ -13,6 +13,7 @@ public:
 
 	TSharedPtr<FUICommandInfo> BeginMorphTargetTool;
 	TSharedPtr<FUICommandInfo> BeginMorphTargetSculptTool;
+	TSharedPtr<FUICommandInfo> BeginMorphTargetMaskTool;
 
 	virtual void RegisterCommands() override;
 };
diff --git a/Engine/Plugins/Experimental/Animation/SkeletalMeshMorphTargetEditingTools/Source/SkeletalMeshMorphTargetEditingTools/SkeletalMeshMorphTargetEditingTools.Build.cs b/Engine/Plugins/Experimental/Animation/SkeletalMeshMorphTargetEditingTools/Source/SkeletalMeshMorphTargetEditingTools/SkeletalMeshMorphTargetEditingTools.Build.cs
--- a/Engine/Plugins/Experimental/Animation/SkeletalMeshMorphTargetEditingTools/Source/SkeletalMeshMorphTargetEditingTools/SkeletalMeshMorphTargetEditingTools.Build.cs
+++ b/Engine/Plugins/Experimental/Animation/SkeletalMeshMorphTargetEditingTools/Source/SkeletalMeshMorphTargetEditingTools/SkeletalMeshMorphTargetEditingTools.Build.cs
@@ -29,7 +29,8 @@ public class SkeletalMeshMorphTargetEditingTools : ModuleRules
 				// ... add other public dependencies that you statically link with here ...
 				"ModelingToolsEditorMode",
 				"SkeletalMeshModelingTools", 
-				"MeshModelingToolsEditorOnly"
+				"MeshModelingToolsEditorOnly", 
+				"MeshModelingToolsEditorOnlyExp"
 			}
 			);
 			
diff --git a/Engine/Plugins/Experimental/Dataflow/Source/DataflowEditor/Private/DataflowEditorTools/DataflowEditorSkinWeightsPaintTool.cpp b/Engine/Plugins/Experimental/Dataflow/Source/DataflowEditor/Private/DataflowEditorTools/DataflowEditorSkinWeightsPaintTool.cpp
--- a/Engine/Plugins/Experimental/Dataflow/Source/DataflowEditor/Private/DataflowEditorTools/DataflowEditorSkinWeightsPaintTool.cpp
+++ b/Engine/Plugins/Experimental/Dataflow/Source/DataflowEditor/Private/DataflowEditorTools/DataflowEditorSkinWeightsPaintTool.cpp
@@ -402,9 +402,9 @@ void UDataflowEditorSkinWeightsPaintTool::PoseMesh(URefSkeletonPoser* Poser)
 
 	auto DeformPreviewMeshFunc = [this, &BoneMatrices](FDynamicMesh3& VisualMesh)
 		{
-			auto WriteFunc = [&VisualMesh](int32 VertID, const FVector& PosedVertPos)
+			auto WriteFunc = [&VisualMesh](SkeletalMeshToolsHelper::FVertInfo VertInfo, const FVector& PosedVertPos)
 				{
-					VisualMesh.SetVertex(VertID, PosedVertPos);
+					VisualMesh.SetVertex(VertInfo.VertID, PosedVertPos);
 				};
 
 			const TMap<FName, float> MorphTargetWeights;
diff --git a/Engine/Plugins/Runtime/MeshModelingToolset/Source/MeshModelingToolsEditorOnly/Private/SkeletalMesh/SkeletalMeshToolsHelper.cpp b/Engine/Plugins/Runtime/MeshModelingToolset/Source/MeshModelingToolsEditorOnly/Private/SkeletalMesh/SkeletalMeshToolsHelper.cpp
--- a/Engine/Plugins/Runtime/MeshModelingToolset/Source/MeshModelingToolsEditorOnly/Private/SkeletalMesh/SkeletalMeshToolsHelper.cpp
+++ b/Engine/Plugins/Runtime/MeshModelingToolset/Source/MeshModelingToolsEditorOnly/Private/SkeletalMesh/SkeletalMeshToolsHelper.cpp
@@ -79,19 +79,26 @@ void SkeletalMeshToolsHelper::GetUnposedMesh(
 }
 
 void SkeletalMeshToolsHelper::GetPosedMesh(
-	TFunctionRef<void(int32, const FVector&)> WriteFunc, 
+	TFunctionRef<void(FVertInfo, const FVector&)> WriteFunc, 
 	const FDynamicMesh3& SourceMesh,
 	const TArray<FMatrix>& BoneMatrices,
 	FName SkinWeightProfile,
-	const TMap<FName, float>& MorphTargetWeights
+	const TMap<FName, float>& MorphTargetWeights,
+	const TArray<int32>& VertArray
 	)
 {
 	using namespace UE::Geometry;
 	using namespace UE::AnimationCore;
 	SkinWeightProfile = SkinWeightProfile == NAME_None ? FSkeletalMeshAttributes::DefaultSkinWeightProfileName : SkinWeightProfile;	
-	FDynamicMeshVertexSkinWeightsAttribute* SkinWeightAttribute = SourceMesh.Attributes()->GetSkinWeightsAttribute(SkinWeightProfile);	
-	ParallelFor(SourceMesh.MaxVertexID(), [&](int32 VertID)
+	FDynamicMeshVertexSkinWeightsAttribute* SkinWeightAttribute = SourceMesh.Attributes()->GetSkinWeightsAttribute(SkinWeightProfile);
+
+	bool bHasVertSelection = VertArray.Num() > 0;
+	int32 NumToProcess = bHasVertSelection ? VertArray.Num() : SourceMesh.MaxVertexID();
+	
+	ParallelFor(NumToProcess, [&](int32 Index)
 		{
+			int32 VertID = bHasVertSelection ? VertArray[Index] : Index;
+			
 			if (!SourceMesh.IsVertex(VertID))
 			{
 				return;
@@ -128,8 +135,13 @@ void SkeletalMeshToolsHelper::GetPosedMesh(
 			}
 			
 			PosedVertPos = SkinMatrix.TransformPosition(PosedVertPos);
-
-			WriteFunc(VertID, PosedVertPos);
+			
+			FVertInfo Info = {
+				bHasVertSelection ? Index : INDEX_NONE,
+				VertID,
+			};
+			
+			WriteFunc(Info, PosedVertPos);
 		});
 }
 
diff --git a/Engine/Plugins/Runtime/MeshModelingToolset/Source/MeshModelingToolsEditorOnly/Public/SkeletalMesh/SkeletalMeshToolsHelper.h b/Engine/Plugins/Runtime/MeshModelingToolset/Source/MeshModelingToolsEditorOnly/Public/SkeletalMesh/SkeletalMeshToolsHelper.h
--- a/Engine/Plugins/Runtime/MeshModelingToolset/Source/MeshModelingToolsEditorOnly/Public/SkeletalMesh/SkeletalMeshToolsHelper.h
+++ b/Engine/Plugins/Runtime/MeshModelingToolset/Source/MeshModelingToolsEditorOnly/Public/SkeletalMesh/SkeletalMeshToolsHelper.h
@@ -33,11 +33,12 @@ namespace SkeletalMeshToolsHelper
 
 	
 	UE_API void GetPosedMesh(
-		TFunctionRef<void(int32, const FVector&)> WriteFunc,
+		TFunctionRef<void(FVertInfo, const FVector&)> WriteFunc,
 		const FDynamicMesh3& SourceMesh,
 		const TArray<FMatrix>& BoneMatrices,
 		FName SkinWeightProfile,
-		const TMap<FName, float>& MorphTargetWeights
+		const TMap<FName, float>& MorphTargetWeights,
+		const TArray<int32>& VertArray = {}
 		);
 	
 	UE_API TArray<FMatrix> ComputeBoneMatrices(
