[SkeletalMeshModelingTools] added a cvar to toggle if base mesh can be edited via geometry selection + set preserveuvflow to be true by default for morph sculpt tool

#Synced-CL 52183294

diff --git a/Engine/Plugins/Animation/SkeletalMeshModelingTools/Source/SkeletalMeshModelingTools/Private/SkeletalMeshModelingToolsEditorMode.cpp b/Engine/Plugins/Animation/SkeletalMeshModelingTools/Source/SkeletalMeshModelingTools/Private/SkeletalMeshModelingToolsEditorMode.cpp
--- a/Engine/Plugins/Animation/SkeletalMeshModelingTools/Source/SkeletalMeshModelingTools/Private/SkeletalMeshModelingToolsEditorMode.cpp
+++ b/Engine/Plugins/Animation/SkeletalMeshModelingTools/Source/SkeletalMeshModelingTools/Private/SkeletalMeshModelingToolsEditorMode.cpp
@@ -183,6 +183,12 @@ public:
 };
 #endif // ENABLE_STYLUS_SUPPORT
 
+bool bAllowEditingBaseMeshViaGeometrySelection = false;
+static FAutoConsoleVariableRef CVarAllowEditingBaseMeshViaGeometrySelection(
+	TEXT("SkeletalMeshModelingTools.AllowEditingBaseMeshViaGeometrySelection"),
+	bAllowEditingBaseMeshViaGeometrySelection,
+	TEXT("If false, base mesh can no longer be edited via geometry selection"));
+
 static void ShowEditorMessage(ELogVerbosity::Type InMessageType, const FText& InMessage)
 {
 	FNotificationInfo Notification(InMessage);
@@ -616,7 +622,7 @@ TArray<FName> USkeletalMeshModelingToolsEditorMode::GetSelectedBones() const
 	return SelectedBones;
 }
 
-FName USkeletalMeshModelingToolsEditorMode::GetEditingMorphTarget()
+FName USkeletalMeshModelingToolsEditorMode::GetEditingMorphTarget() const
 {
 	if (!GetCurrentEditingCache())
 	{
@@ -1497,6 +1503,11 @@ bool USkeletalMeshModelingToolsEditorMode::ShouldDrawWidget() const
 {
 	if (HasActiveGeometrySelection())
 	{
+		if (GetEditingMorphTarget() == NAME_None)
+		{
+			return bAllowEditingBaseMeshViaGeometrySelection;
+		}
+			
 		return true;
 	}
 	
diff --git a/Engine/Plugins/Animation/SkeletalMeshModelingTools/Source/SkeletalMeshModelingTools/Private/SkeletalMeshModelingToolsEditorMode.h b/Engine/Plugins/Animation/SkeletalMeshModelingTools/Source/SkeletalMeshModelingTools/Private/SkeletalMeshModelingToolsEditorMode.h
--- a/Engine/Plugins/Animation/SkeletalMeshModelingTools/Source/SkeletalMeshModelingTools/Private/SkeletalMeshModelingToolsEditorMode.h
+++ b/Engine/Plugins/Animation/SkeletalMeshModelingTools/Source/SkeletalMeshModelingTools/Private/SkeletalMeshModelingToolsEditorMode.h
@@ -159,7 +159,7 @@ public:
 	void SetSelectedBones(const TArray<FName>& InSelectedBones);
 	TArray<FName> GetSelectedBones() const;
 	
-	FName GetEditingMorphTarget();
+	FName GetEditingMorphTarget() const;
 	void HandleSetEditingMorphTarget(FName InMorphTarget);
 
 	TArray<FName> GetMorphTargets();
diff --git a/Engine/Plugins/Experimental/Animation/SkeletalMeshMorphTargetEditingTools/Source/SkeletalMeshMorphTargetEditingTools/Private/MorphTargetVertexSculptTool.cpp b/Engine/Plugins/Experimental/Animation/SkeletalMeshMorphTargetEditingTools/Source/SkeletalMeshMorphTargetEditingTools/Private/MorphTargetVertexSculptTool.cpp
--- a/Engine/Plugins/Experimental/Animation/SkeletalMeshMorphTargetEditingTools/Source/SkeletalMeshMorphTargetEditingTools/Private/MorphTargetVertexSculptTool.cpp
+++ b/Engine/Plugins/Experimental/Animation/SkeletalMeshMorphTargetEditingTools/Source/SkeletalMeshMorphTargetEditingTools/Private/MorphTargetVertexSculptTool.cpp
@@ -29,6 +29,7 @@
 #include "TargetInterfaces/MeshDescriptionProvider.h"
 #include "TargetInterfaces/SkeletonProvider.h"
 #include "DynamicMesh/MeshNormals.h"
+#include "Sculpting/MeshSmoothingBrushOps.h"
 
 #include UE_INLINE_GENERATED_CPP_BY_NAME(MorphTargetVertexSculptTool)
 
@@ -182,6 +183,13 @@ void UMorphTargetVertexSculptTool::RegisterBrushes()
 		LOCTEXT("EraseSculptLayerBrushTypeName", "EraseSculptLayer"),
 		MakeUnique<FLambdaMeshSculptBrushOpFactory>([this]() { return MakeUnique<FEraseMorphTargetBrushOp>(GetMeshWithoutCurrentMorphFunc);}),
 		NewObject<UEraseMorphTargetBrushOpProps>(this));
+	
+	USecondarySmoothBrushOpProps* SecondarySmoothProps = Cast<USecondarySmoothBrushOpProps>(SecondaryBrushOpPropSets.FindChecked(0));
+	SecondarySmoothProps->bPreserveUVFlow = true;
+	USmoothBrushOpProps* SmoothProps = Cast<USmoothBrushOpProps>(BrushOpPropSets.FindChecked((int32)EMeshVertexSculptBrushType::Smooth));
+	SmoothProps->bPreserveUVFlow = true;
+	USmoothFillBrushOpProps* SmoothFillProps = Cast<USmoothFillBrushOpProps>(BrushOpPropSets.FindChecked((int32)EMeshVertexSculptBrushType::SmoothFill));
+	SmoothFillProps->bPreserveUVFlow = true;
 }
 
 void UMorphTargetVertexSculptTool::Shutdown(EToolShutdownType ShutdownType)
