Move IMorphTargetEditingToolInterface to SkeletalMeshModelingTools and auto-restart morph editing tools on editing morph change

#Synced-CL 52437958

diff --git a/Engine/Plugins/Animation/SkeletalMeshModelingTools/Source/SkeletalMeshModelingTools/Private/SkeletalMeshModelingToolsEditorMode.cpp b/Engine/Plugins/Animation/SkeletalMeshModelingTools/Source/SkeletalMeshModelingTools/Private/SkeletalMeshModelingToolsEditorMode.cpp
--- a/Engine/Plugins/Animation/SkeletalMeshModelingTools/Source/SkeletalMeshModelingTools/Private/SkeletalMeshModelingToolsEditorMode.cpp
+++ b/Engine/Plugins/Animation/SkeletalMeshModelingTools/Source/SkeletalMeshModelingTools/Private/SkeletalMeshModelingToolsEditorMode.cpp
@@ -3,6 +3,7 @@
 #include "SkeletalMeshModelingToolsEditorMode.h"
 
 #include "AttributeEditorTool.h"
+#include "ToolInterfaces/IMorphTargetEditingToolInterface.h"
 #include "SkeletalMeshModelingToolsEditorModeToolkit.h"
 #include "SkeletalMeshModelingToolsCommands.h"
 #include "SkeletalMeshGizmoUtils.h"
@@ -645,6 +646,21 @@ void USkeletalMeshModelingToolsEditorMode::HandleSetEditingMorphTarget(FName InM
 	// Such that EditingMorphTarget is recorded for undo/redo 
 	Modify();
 	EditingMorphTarget = InMorphTarget;
+
+	// Restart active morph target editing tool so it picks up the new editing target
+	UInteractiveToolManager* ToolManager = GetInteractiveToolsContext()->ToolManager;
+	if (ToolManager->HasActiveTool(EToolSide::Left))
+	{
+		UInteractiveTool* ActiveTool = ToolManager->GetActiveTool(EToolSide::Left);
+		if (Cast<IMorphTargetEditingToolInterface>(ActiveTool))
+		{
+			EToolShutdownType ShutdownType = ToolManager->CanAcceptActiveTool(EToolSide::Left)
+				? EToolShutdownType::Accept
+				: EToolShutdownType::Completed;
+			ToolManager->DeactivateTool(EToolSide::Left, ShutdownType);
+			ToolManager->ActivateTool(EToolSide::Left);
+		}
+	}
 }
 
 
diff --git a/Engine/Plugins/Animation/SkeletalMeshModelingTools/Source/SkeletalMeshModelingTools/Private/ToolInterfaces/IMorphTargetEditingToolInterface.cpp b/Engine/Plugins/Animation/SkeletalMeshModelingTools/Source/SkeletalMeshModelingTools/Private/ToolInterfaces/IMorphTargetEditingToolInterface.cpp
new file mode 100644
--- /dev/null
+++ b/Engine/Plugins/Animation/SkeletalMeshModelingTools/Source/SkeletalMeshModelingTools/Private/ToolInterfaces/IMorphTargetEditingToolInterface.cpp
@@ -0,0 +1,28 @@
+﻿// Copyright Epic Games, Inc. All Rights Reserved.
+
+#include "ToolInterfaces/IMorphTargetEditingToolInterface.h"
+
+#include "ContextObjectStore.h"
+#include "InteractiveToolManager.h"
+#include "SkeletalMesh/SkeletalMeshEditingInterface.h"
+
+#include UE_INLINE_GENERATED_CPP_BY_NAME(IMorphTargetEditingToolInterface)
+
+void IMorphTargetEditingToolInterface::SetupMorphEditingToolCommon()
+{
+	UInteractiveTool* Tool = CastChecked<UInteractiveTool>(this);
+	USkeletalMeshEditorContextObjectBase* EditorContext = Tool->GetToolManager()->GetContextObjectStore()->FindContext<USkeletalMeshEditorContextObjectBase>();
+
+	
+	auto SetupFunction = [&](UMorphTargetEditingToolProperties* InProperties)
+	{
+		InProperties->EditMorphTargetName = EditorContext->GetEditingMorphTarget();
+	};
+	
+	SetupCommonProperties(SetupFunction);
+}
+
+void IMorphTargetEditingToolInterface::ShutdownMorphEditingToolCommon()
+{
+
+}
diff --git a/Engine/Plugins/Animation/SkeletalMeshModelingTools/Source/SkeletalMeshModelingTools/Public/ToolInterfaces/IMorphTargetEditingToolInterface.h b/Engine/Plugins/Animation/SkeletalMeshModelingTools/Source/SkeletalMeshModelingTools/Public/ToolInterfaces/IMorphTargetEditingToolInterface.h
new file mode 100644
--- /dev/null
+++ b/Engine/Plugins/Animation/SkeletalMeshModelingTools/Source/SkeletalMeshModelingTools/Public/ToolInterfaces/IMorphTargetEditingToolInterface.h
@@ -0,0 +1,39 @@
+﻿// Copyright Epic Games, Inc. All Rights Reserved.
+
+#pragma once
+
+#include "CoreMinimal.h"
+#include "InteractiveTool.h"
+#include "UObject/Interface.h"
+
+#include "IMorphTargetEditingToolInterface.generated.h"
+
+#define UE_API SKELETALMESHMODELINGTOOLS_API
+
+class UMorphTargetEditingToolProperties;
+
+UINTERFACE(MinimalAPI)
+class UMorphTargetEditingToolInterface : public UInterface
+{
+	GENERATED_BODY()
+};
+
+class IMorphTargetEditingToolInterface : public IInterface
+{
+	GENERATED_BODY()
+public:
+	UE_API virtual void SetupMorphEditingToolCommon();
+	UE_API virtual void ShutdownMorphEditingToolCommon();
+	virtual void SetupCommonProperties(const TFunction<void(UMorphTargetEditingToolProperties*)>& InSetupFunction) = 0;
+};
+
+UCLASS()
+class UMorphTargetEditingToolProperties : public UInteractiveToolPropertySet
+{
+	GENERATED_BODY()
+public:
+	UPROPERTY(VisibleAnywhere, Category = "MorphTarget", meta = (NoResetToDefault, DisplayName = "Name"))
+	FName EditMorphTargetName;
+};
+
+#undef UE_API
diff --git a/Engine/Plugins/Experimental/Animation/SkeletalMeshMorphTargetEditingTools/Source/SkeletalMeshMorphTargetEditingTools/Private/IMorphTargetEditingToolInterface.cpp b/Engine/Plugins/Experimental/Animation/SkeletalMeshMorphTargetEditingTools/Source/SkeletalMeshMorphTargetEditingTools/Private/IMorphTargetEditingToolInterface.cpp
deleted file mode 100644
diff --git a/Engine/Plugins/Experimental/Animation/SkeletalMeshMorphTargetEditingTools/Source/SkeletalMeshMorphTargetEditingTools/Private/MorphTargetEditingToolProperties.cpp b/Engine/Plugins/Experimental/Animation/SkeletalMeshMorphTargetEditingTools/Source/SkeletalMeshMorphTargetEditingTools/Private/MorphTargetEditingToolProperties.cpp
deleted file mode 100644
diff --git a/Engine/Plugins/Experimental/Animation/SkeletalMeshMorphTargetEditingTools/Source/SkeletalMeshMorphTargetEditingTools/Private/MorphTargetVertexSculptTool.cpp b/Engine/Plugins/Experimental/Animation/SkeletalMeshMorphTargetEditingTools/Source/SkeletalMeshMorphTargetEditingTools/Private/MorphTargetVertexSculptTool.cpp
--- a/Engine/Plugins/Experimental/Animation/SkeletalMeshMorphTargetEditingTools/Source/SkeletalMeshMorphTargetEditingTools/Private/MorphTargetVertexSculptTool.cpp
+++ b/Engine/Plugins/Experimental/Animation/SkeletalMeshMorphTargetEditingTools/Source/SkeletalMeshMorphTargetEditingTools/Private/MorphTargetVertexSculptTool.cpp
@@ -175,7 +175,14 @@ void UMorphTargetVertexSculptTool::RegisterBrushes()
 {
 	Super::RegisterBrushes();
 	
-	GetMeshWithoutCurrentMorphFunc = [this](){return &MeshWithoutEditingMorph;};
+	GetMeshWithoutCurrentMorphFunc = [this]()
+	{
+		if (bMeshWithoutEditingMorphDirty)
+		{
+			UpdateMeshWithoutEditingMorph();
+		}
+		return &MeshWithoutEditingMorph;
+	};
 
 	// Had to hijack the EraseSculptLayer identifier from base mesh vertex sculpt tool for our erase morph target tool since it is the simplest way to get an icon for the tool.
 	RegisterBrushType(
@@ -183,13 +190,6 @@ void UMorphTargetVertexSculptTool::RegisterBrushes()
 		LOCTEXT("EraseSculptLayerBrushTypeName", "EraseSculptLayer"),
 		MakeUnique<FLambdaMeshSculptBrushOpFactory>([this]() { return MakeUnique<FEraseMorphTargetBrushOp>(GetMeshWithoutCurrentMorphFunc);}),
 		NewObject<UEraseMorphTargetBrushOpProps>(this));
-	
-	USecondarySmoothBrushOpProps* SecondarySmoothProps = Cast<USecondarySmoothBrushOpProps>(SecondaryBrushOpPropSets.FindChecked(0));
-	SecondarySmoothProps->bPreserveUVFlow = true;
-	USmoothBrushOpProps* SmoothProps = Cast<USmoothBrushOpProps>(BrushOpPropSets.FindChecked((int32)EMeshVertexSculptBrushType::Smooth));
-	SmoothProps->bPreserveUVFlow = true;
-	USmoothFillBrushOpProps* SmoothFillProps = Cast<USmoothFillBrushOpProps>(BrushOpPropSets.FindChecked((int32)EMeshVertexSculptBrushType::SmoothFill));
-	SmoothFillProps->bPreserveUVFlow = true;
 }
 
 void UMorphTargetVertexSculptTool::Shutdown(EToolShutdownType ShutdownType)
@@ -333,71 +333,80 @@ void UMorphTargetVertexSculptTool::OnEndStroke()
 	float EditingMorphTargetWeight = MorphTargetWeights[ToolMorphTargetName];
 	
 	// Don't update the cache if the morph cannot be extracted from the mesh
-	if (!FMath::IsNearlyZero(EditingMorphTargetWeight) && !ActiveVertexChange->Change->Vertices.IsEmpty())
+	if (FMath::IsNearlyZero(EditingMorphTargetWeight))
 	{
-		FMeshMorphTargetChange* MorphTargetChange = new FMeshMorphTargetChange(); 
-		MorphTargetChange->MorphTargetName = ToolMorphTargetName;
-		MorphTargetChange->Vertices = MoveTemp(ActiveVertexChange->Change->Vertices);
-		MorphTargetChange->OldDeltas.SetNumUninitialized(MorphTargetChange->Vertices.Num());
-		MorphTargetChange->NewDeltas.SetNumUninitialized(MorphTargetChange->Vertices.Num());
-		
-		// Exclude the morph target we are trying to extract
-		MorphTargetWeights.Remove(ToolMorphTargetName);
-		
-		using namespace SkeletalMeshToolsHelper;
-		TArray<FMatrix> BoneMatrices = ComputeBoneMatrices(
+		bFastDeformSculptMesh = true;
+		bFullRefreshSculptMesh = true;
+	}
+	else
+	{
+		// Extract morph delta if possible
+		if (!ActiveVertexChange->Change->Vertices.IsEmpty())
+		{
+			FMeshMorphTargetChange* MorphTargetChange = new FMeshMorphTargetChange();
+			MorphTargetChange->MorphTargetName = ToolMorphTargetName;
+			MorphTargetChange->Vertices = MoveTemp(ActiveVertexChange->Change->Vertices);
+			MorphTargetChange->OldDeltas.SetNumUninitialized(MorphTargetChange->Vertices.Num());
+			MorphTargetChange->NewDeltas.SetNumUninitialized(MorphTargetChange->Vertices.Num());
+
+			// Exclude the morph target we are trying to extract
+			MorphTargetWeights.Remove(ToolMorphTargetName);
+
+			using namespace SkeletalMeshToolsHelper;
+			TArray<FMatrix> BoneMatrices = ComputeBoneMatrices(
 				ComponentSpaceTransformsRefPose,
 				ComponentSpaceTransforms
-				);
+			);
 
-		ToolDynamicMesh->EditMesh([&](FDynamicMesh3& Mesh)
-			{
-				FDynamicMeshMorphTargetAttribute* MorphTargetAttribute = Mesh.Attributes()->GetMorphTargetAttribute(ToolMorphTargetName);
-				
-				auto WriteFunction = [&](FVertInfo VertInfo, const FVector& UnposedVertPos)
-					{
-						int32 VertID = VertInfo.VertID;
-						FVector NewDelta = UnposedVertPos - FVector(Mesh.GetVertex(VertID));
-						NewDelta /= EditingMorphTargetWeight;
-						
-						FVector OldDelta;
-						MorphTargetAttribute->GetValue(VertID, OldDelta);
-						MorphTargetAttribute->SetValue(VertID, NewDelta);
-
-						MorphTargetChange->OldDeltas[VertInfo.VertArrayIndex] = OldDelta;
-						MorphTargetChange->NewDeltas[VertInfo.VertArrayIndex] = NewDelta;
-					};
-
-				GetUnposedMesh(WriteFunction, *GetSculptMesh(), Mesh, BoneMatrices, NAME_None, MorphTargetWeights, MorphTargetChange->Vertices);
-			});
-		
-		
-		TUniquePtr<TWrappedToolCommandChange<FMeshMorphTargetChange>> NewChange = MakeUnique<TWrappedToolCommandChange<FMeshMorphTargetChange>>();
-		
-		NewChange->WrappedChange.Reset(MorphTargetChange);
-		NewChange->BeforeModify = [this](bool bRevert)
-		{
-			WaitForPendingUndoRedoUpdate();
-		};
+			ToolDynamicMesh->EditMesh([&](FDynamicMesh3& Mesh)
+				{
+					FDynamicMeshMorphTargetAttribute* MorphTargetAttribute = Mesh.Attributes()->GetMorphTargetAttribute(ToolMorphTargetName);
 
-		
-		NewChange->AfterModify = [this](bool bRevert)
-		{
-			bFastDeformSculptMesh = true;
-			bFullRefreshSculptMesh = true;
-		};
+					auto WriteFunction = [&](FVertInfo VertInfo, const FVector& UnposedVertPos)
+						{
+							int32 VertID = VertInfo.VertID;
+							FVector NewDelta = UnposedVertPos - FVector(Mesh.GetVertex(VertID));
+							NewDelta /= EditingMorphTargetWeight;
 
-		GetToolManager()->EmitObjectChange(ToolDynamicMesh, MoveTemp(NewChange), LOCTEXT("VertexSculptChange", "Brush Stroke"));
-		if (bMeshSymmetryIsValid && bApplySymmetry == false)
-		{
-			// if we end a stroke while symmetry is possible but disabled, we now have to assume that symmetry is no longer possible
-			GetToolManager()->EmitObjectChange(this, MakeUnique<FMorphTargetVertexSculptNonSymmetricChange>(), LOCTEXT("DisableSymmetryChange", "Disable Symmetry"));
-			bMeshSymmetryIsValid = false;
-			SymmetryProperties->bSymmetryCanBeEnabled = bMeshSymmetryIsValid;
+							FVector OldDelta;
+							MorphTargetAttribute->GetValue(VertID, OldDelta);
+							MorphTargetAttribute->SetValue(VertID, NewDelta);
+
+							MorphTargetChange->OldDeltas[VertInfo.VertArrayIndex] = OldDelta;
+							MorphTargetChange->NewDeltas[VertInfo.VertArrayIndex] = NewDelta;
+						};
+
+					GetUnposedMesh(WriteFunction, *GetSculptMesh(), Mesh, BoneMatrices, NAME_None, MorphTargetWeights, MorphTargetChange->Vertices);
+				});
+
+
+			TUniquePtr<TWrappedToolCommandChange<FMeshMorphTargetChange>> NewChange = MakeUnique<TWrappedToolCommandChange<FMeshMorphTargetChange>>();
+
+			NewChange->WrappedChange.Reset(MorphTargetChange);
+			NewChange->BeforeModify = [this](bool bRevert)
+				{
+					WaitForPendingUndoRedoUpdate();
+				};
+
+
+			NewChange->AfterModify = [this](bool bRevert)
+				{
+					bFastDeformSculptMesh = true;
+					bFullRefreshSculptMesh = true;
+				};
+
+			GetToolManager()->EmitObjectChange(ToolDynamicMesh, MoveTemp(NewChange), LOCTEXT("VertexSculptChange", "Brush Stroke"));
+			if (bMeshSymmetryIsValid && bApplySymmetry == false)
+			{
+				// if we end a stroke while symmetry is possible but disabled, we now have to assume that symmetry is no longer possible
+				GetToolManager()->EmitObjectChange(this, MakeUnique<FMorphTargetVertexSculptNonSymmetricChange>(),
+					LOCTEXT("DisableSymmetryChange", "Disable Symmetry"));
+				bMeshSymmetryIsValid = false;
+				SymmetryProperties->bSymmetryCanBeEnabled = bMeshSymmetryIsValid;
+			}
 		}
 	}
 	
-	
 	LongTransactions.Close(GetToolManager());
 
 	delete ActiveVertexChange;
@@ -432,18 +441,9 @@ void UMorphTargetVertexSculptTool::PoseSculptMesh(const TArray<FTransform>& Comp
 
 	ToolDynamicMesh->ProcessMesh([&](const FDynamicMesh3& InMesh)
 		{
-			FDynamicMeshMorphTargetAttribute* MorphTargetAttribute = InMesh.Attributes()->GetMorphTargetAttribute(ToolMorphTargetName);
-
-			float EditingMorphTargetWeight = MorphTargetWeights[ToolMorphTargetName];
-	
 			auto WriteFunc = [&](FVertInfo VertInfo, const FVector& PosedVertPos)
 				{
 					SculptMesh.SetVertex(VertInfo.VertID, PosedVertPos);
-
-					FVector Delta;
-					MorphTargetAttribute->GetValue(VertInfo.VertID, Delta);
-			
-					MeshWithoutEditingMorph.SetVertex(VertInfo.VertID,  PosedVertPos - (Delta * EditingMorphTargetWeight));
 				};
 
 			GetPosedMesh(WriteFunc, InMesh, BoneMatrices, NAME_None, MorphTargetWeights);
@@ -451,10 +451,10 @@ void UMorphTargetVertexSculptTool::PoseSculptMesh(const TArray<FTransform>& Comp
 			
 			ComponentSpaceTransformsForPosedMesh = ComponentSpaceTransforms;
 			MorphTargetWeightsForPosedMesh = MorphTargetWeights;
-		
-
 		});
 
+	bMeshWithoutEditingMorphDirty = true;
+
 	constexpr bool bNormal = true;
 	DynamicMeshComponent->FastNotifyPositionsUpdated(bNormal);
 
@@ -462,5 +462,31 @@ void UMorphTargetVertexSculptTool::PoseSculptMesh(const TArray<FTransform>& Comp
 }
 
 
+void UMorphTargetVertexSculptTool::UpdateMeshWithoutEditingMorph()
+{
+	using namespace SkeletalMeshToolsHelper;
+
+	WaitForPendingStampUpdateConst();
+	WaitForPendingUndoRedoUpdate();
+
+	TArray<FMatrix> BoneMatrices = ComputeBoneMatrices(ComponentSpaceTransformsRefPose, ComponentSpaceTransformsForPosedMesh);
+
+	// Build weights excluding the editing morph target
+	TMap<FName, float> WeightsWithoutEditingMorph = MorphTargetWeightsForPosedMesh;
+	WeightsWithoutEditingMorph.Remove(ToolMorphTargetName);
+
+	ToolDynamicMesh->ProcessMesh([&](const FDynamicMesh3& InMesh)
+		{
+			auto WriteFunc = [&](FVertInfo VertInfo, const FVector& PosedVertPos)
+				{
+					MeshWithoutEditingMorph.SetVertex(VertInfo.VertID, PosedVertPos);
+				};
+
+			GetPosedMesh(WriteFunc, InMesh, BoneMatrices, NAME_None, WeightsWithoutEditingMorph);
+		});
+
+	bMeshWithoutEditingMorphDirty = false;
+}
+
 #undef LOCTEXT_NAMESPACE
 
diff --git a/Engine/Plugins/Experimental/Animation/SkeletalMeshMorphTargetEditingTools/Source/SkeletalMeshMorphTargetEditingTools/Public/IMorphTargetEditingToolInterface.h b/Engine/Plugins/Experimental/Animation/SkeletalMeshMorphTargetEditingTools/Source/SkeletalMeshMorphTargetEditingTools/Public/IMorphTargetEditingToolInterface.h
deleted file mode 100644
diff --git a/Engine/Plugins/Experimental/Animation/SkeletalMeshMorphTargetEditingTools/Source/SkeletalMeshMorphTargetEditingTools/Public/MorphTargetEditingToolProperties.h b/Engine/Plugins/Experimental/Animation/SkeletalMeshMorphTargetEditingTools/Source/SkeletalMeshMorphTargetEditingTools/Public/MorphTargetEditingToolProperties.h
deleted file mode 100644
diff --git a/Engine/Plugins/Experimental/Animation/SkeletalMeshMorphTargetEditingTools/Source/SkeletalMeshMorphTargetEditingTools/Public/MorphTargetMaskTool.h b/Engine/Plugins/Experimental/Animation/SkeletalMeshMorphTargetEditingTools/Source/SkeletalMeshMorphTargetEditingTools/Public/MorphTargetMaskTool.h
--- a/Engine/Plugins/Experimental/Animation/SkeletalMeshMorphTargetEditingTools/Source/SkeletalMeshMorphTargetEditingTools/Public/MorphTargetMaskTool.h
+++ b/Engine/Plugins/Experimental/Animation/SkeletalMeshMorphTargetEditingTools/Source/SkeletalMeshMorphTargetEditingTools/Public/MorphTargetMaskTool.h
@@ -2,8 +2,7 @@
 
 #pragma once
 
-#include "IMorphTargetEditingToolInterface.h"
-#include "MorphTargetEditingToolProperties.h"
+#include "ToolInterfaces/IMorphTargetEditingToolInterface.h"
 
 #include "MeshVertexAttributePaintToolBase.h"
 #include "SkeletalMesh/SkeletalMeshToolsHelper.h"
diff --git a/Engine/Plugins/Experimental/Animation/SkeletalMeshMorphTargetEditingTools/Source/SkeletalMeshMorphTargetEditingTools/Public/MorphTargetVertexSculptTool.h b/Engine/Plugins/Experimental/Animation/SkeletalMeshMorphTargetEditingTools/Source/SkeletalMeshMorphTargetEditingTools/Public/MorphTargetVertexSculptTool.h
--- a/Engine/Plugins/Experimental/Animation/SkeletalMeshMorphTargetEditingTools/Source/SkeletalMeshMorphTargetEditingTools/Public/MorphTargetVertexSculptTool.h
+++ b/Engine/Plugins/Experimental/Animation/SkeletalMeshMorphTargetEditingTools/Source/SkeletalMeshMorphTargetEditingTools/Public/MorphTargetVertexSculptTool.h
@@ -2,11 +2,10 @@
 
 #pragma once
 
-#include "IMorphTargetEditingToolInterface.h"
+#include "ToolInterfaces/IMorphTargetEditingToolInterface.h"
 #include "MeshVertexSculptTool.h"
 #include "SkeletalMesh/SkeletalMeshEditingInterface.h"
 #include "BaseTools/MultiSelectionMeshEditingTool.h"
-#include "MorphTargetEditingToolProperties.h"
 #include "MeshDescription.h"
 #include "ReferenceSkeleton.h"
 #include "SkeletalMesh/SkeletalMeshToolsHelper.h"
@@ -78,6 +77,8 @@ protected:
 
 	TFunction<const FDynamicMesh3*()> GetMeshWithoutCurrentMorphFunc;
 	FDynamicMesh3 MeshWithoutEditingMorph;
+	bool bMeshWithoutEditingMorphDirty = true;
+	void UpdateMeshWithoutEditingMorph();
 	
 	SkeletalMeshToolsHelper::FPoseChangeDetector PoseChangeDetector;
 
