[SkeletalMeshModelingTools][MorphTargetMaskTool] Add geometry isolation support to Morph Target Mask Tool

Expose isolation triangles via EditorContext. Mask tool builds submesh for display/paint,
remaps vertex IDs when committing weights back to the full mesh.

#Synced-CL 52717952

diff --git a/Engine/Plugins/Animation/SkeletalMeshModelingTools/Source/SkeletalMeshModelingTools/Private/SkeletalMeshEditingCache.cpp b/Engine/Plugins/Animation/SkeletalMeshModelingTools/Source/SkeletalMeshModelingTools/Private/SkeletalMeshEditingCache.cpp
--- a/Engine/Plugins/Animation/SkeletalMeshModelingTools/Source/SkeletalMeshModelingTools/Private/SkeletalMeshEditingCache.cpp
+++ b/Engine/Plugins/Animation/SkeletalMeshModelingTools/Source/SkeletalMeshModelingTools/Private/SkeletalMeshEditingCache.cpp
@@ -1021,6 +1021,12 @@ bool USkeletalMeshEditingCache::HasSavedIsolation() const
 	return SavedIsolationTriangles.IsSet();
 }
 
+const TArray<int32>& USkeletalMeshEditingCache::GetSavedIsolationTriangles() const
+{
+	static const TArray<int32> Empty;
+	return SavedIsolationTriangles.IsSet() ? SavedIsolationTriangles.GetValue() : Empty;
+}
+
 void USkeletalMeshEditingCache::RestoreSavedIsolation()
 {
 	if (SavedIsolationTriangles.IsSet())
diff --git a/Engine/Plugins/Animation/SkeletalMeshModelingTools/Source/SkeletalMeshModelingTools/Private/SkeletalMeshEditingCache.h b/Engine/Plugins/Animation/SkeletalMeshModelingTools/Source/SkeletalMeshModelingTools/Private/SkeletalMeshEditingCache.h
--- a/Engine/Plugins/Animation/SkeletalMeshModelingTools/Source/SkeletalMeshModelingTools/Private/SkeletalMeshEditingCache.h
+++ b/Engine/Plugins/Animation/SkeletalMeshModelingTools/Source/SkeletalMeshModelingTools/Private/SkeletalMeshEditingCache.h
@@ -131,6 +131,7 @@ public:
 
 	// Saved isolation — allows temporarily showing full mesh and restoring later (e.g. during tool use)
 	bool HasSavedIsolation() const;
+	const TArray<int32>& GetSavedIsolationTriangles() const;
 	void RestoreSavedIsolation();
 	void DiscardSavedIsolation();
 
diff --git a/Engine/Plugins/Animation/SkeletalMeshModelingTools/Source/SkeletalMeshModelingTools/Private/SkeletalMeshEditorUtils.cpp b/Engine/Plugins/Animation/SkeletalMeshModelingTools/Source/SkeletalMeshModelingTools/Private/SkeletalMeshEditorUtils.cpp
--- a/Engine/Plugins/Animation/SkeletalMeshModelingTools/Source/SkeletalMeshModelingTools/Private/SkeletalMeshEditorUtils.cpp
+++ b/Engine/Plugins/Animation/SkeletalMeshModelingTools/Source/SkeletalMeshModelingTools/Private/SkeletalMeshEditorUtils.cpp
@@ -140,6 +140,11 @@ void USkeletalMeshEditorContextObject::NotifyMorphTargetEdited()
 	return EditorMode->GetCurrentEditingCache()->HandleMorphTargetEdited( EditorMode->GetEditingMorphTarget() );
 }
 
+const TArray<int32>& USkeletalMeshEditorContextObject::GetIsolatedTriangles() const
+{
+	return EditorMode->GetCurrentEditingCache()->GetSavedIsolationTriangles();
+}
+
 void USkeletalMeshEditorContextObject::BindTo(ISkeletalMeshEditingInterface* InEditingInterface)
 {
 	if (!InEditingInterface)
diff --git a/Engine/Plugins/Animation/SkeletalMeshModelingTools/Source/SkeletalMeshModelingTools/Private/SkeletalMeshEditorUtils.h b/Engine/Plugins/Animation/SkeletalMeshModelingTools/Source/SkeletalMeshModelingTools/Private/SkeletalMeshEditorUtils.h
--- a/Engine/Plugins/Animation/SkeletalMeshModelingTools/Source/SkeletalMeshModelingTools/Private/SkeletalMeshEditorUtils.h
+++ b/Engine/Plugins/Animation/SkeletalMeshModelingTools/Source/SkeletalMeshModelingTools/Private/SkeletalMeshEditorUtils.h
@@ -57,6 +57,7 @@ public:
 	UE_API virtual FName GetEditingMorphTarget() override;
 	UE_API virtual TMap<FName, float> GetMorphTargetWeights() override;
 	UE_API virtual void NotifyMorphTargetEdited() override;
+	UE_API virtual const TArray<int32>& GetIsolatedTriangles() const override;
 	
 
 	UE_API virtual void BindTo(ISkeletalMeshEditingInterface* InEditingInterface) override;
diff --git a/Engine/Plugins/Experimental/Animation/SkeletalMeshMorphTargetEditingTools/Source/SkeletalMeshMorphTargetEditingTools/Private/MorphTargetMaskTool.cpp b/Engine/Plugins/Experimental/Animation/SkeletalMeshMorphTargetEditingTools/Source/SkeletalMeshMorphTargetEditingTools/Private/MorphTargetMaskTool.cpp
--- a/Engine/Plugins/Experimental/Animation/SkeletalMeshMorphTargetEditingTools/Source/SkeletalMeshMorphTargetEditingTools/Private/MorphTargetMaskTool.cpp
+++ b/Engine/Plugins/Experimental/Animation/SkeletalMeshMorphTargetEditingTools/Source/SkeletalMeshMorphTargetEditingTools/Private/MorphTargetMaskTool.cpp
@@ -176,7 +176,18 @@ bool UMorphTargetMaskTool::SetupToolMesh(FDynamicMesh3& InOutToolMesh, int32& Ou
 			
 		MaskToolCommitMesh.Attributes()->AttachMorphTargetAttribute(ToolMorphTargetBackupAttributeName, ToolMorphTargetBackupAttribute);	
 	}
-	
+
+	// Build isolated submesh if isolation was active
+	if (EditorContext.IsValid())
+	{
+		const TArray<int32>& IsolatedTriangles = EditorContext->GetIsolatedTriangles();
+		if (!IsolatedTriangles.IsEmpty())
+		{
+			IsolationSubmesh = FDynamicSubmesh3(&MaskToolCommitMesh, IsolatedTriangles);
+			InOutToolMesh = IsolationSubmesh->GetSubmesh();
+		}
+	}
+
 	// Create a temp weight layer to paint into
 	const int32 NumAttributeLayers = InOutToolMesh.Attributes()->NumWeightLayers();
 	InOutToolMesh.Attributes()->SetNumWeightLayers(NumAttributeLayers + 1);
@@ -396,18 +407,40 @@ void UMorphTargetMaskTool::CommitToolMesh(FDynamicMesh3& InToolMesh)
 		if (TargetWeightMap)
 		{
 			FDynamicMeshWeightAttribute* SourceWeightMap = InToolMesh.Attributes()->GetWeightLayer(CurrentAttributeIndex);
-	
-			ParallelFor(MaskToolCommitMesh.MaxVertexID(), [this, SourceWeightMap, TargetWeightMap](int32 VertID)
+
+			if (IsolationSubmesh)
 			{
-				if (!MaskToolCommitMesh.IsVertex(VertID))
+				// Iterate submesh vertices, remap to full mesh
+				ParallelFor(InToolMesh.MaxVertexID(), [this, SourceWeightMap, TargetWeightMap, &InToolMesh](int32 SubVID)
 				{
-					return;
-				}
+					if (!InToolMesh.IsVertex(SubVID))
+					{
+						return;
+					}
+
+					int32 FullVID = IsolationSubmesh->MapVertexToBaseMesh(SubVID);
+					if (FullVID != FDynamicMesh3::InvalidID && MaskToolCommitMesh.IsVertex(FullVID))
+					{
+						float Weight = 0.0f;
+						SourceWeightMap->GetValue(SubVID, &Weight);
+						TargetWeightMap->SetValue(FullVID, &Weight);
+					}
+				});
+			}
+			else
+			{
+				ParallelFor(MaskToolCommitMesh.MaxVertexID(), [this, SourceWeightMap, TargetWeightMap](int32 VertID)
+				{
+					if (!MaskToolCommitMesh.IsVertex(VertID))
+					{
+						return;
+					}
 				
-				float Weight = 0.0f;
-				SourceWeightMap->GetValue(VertID, &Weight);
-				TargetWeightMap->SetValue(VertID, &Weight);
-			});	
+					float Weight = 0.0f;
+					SourceWeightMap->GetValue(VertID, &Weight);
+					TargetWeightMap->SetValue(VertID, &Weight);
+				});
+			}
 		}
 	}
 
@@ -444,26 +477,29 @@ void UMorphTargetMaskTool::UpdatePreview(const TSet<int32>* TrianglesToUpdate, c
 	FName BaseMorphTargetAttributeName = GetBaseMorphTarget() == ToolMorphTargetName ? ToolMorphTargetBackupAttributeName : GetBaseMorphTarget();  
 	FDynamicMeshMorphTargetAttribute* BaseToolMorphTarget = MaskToolCommitMesh.Attributes()->GetMorphTargetAttribute(BaseMorphTargetAttributeName);
 
-	int32 NumVerts = VerticesToUpdate ? VerticesToUpdate->Num() : MaskToolCommitMesh.MaxVertexID();
+	int32 NumVerts = VerticesToUpdate ? VerticesToUpdate->Num() : GetSculptMesh()->MaxVertexID();
 	
 	ParallelFor(NumVerts, [VerticesToUpdate, this, ToolMorphTargetCache, BaseToolMorphTarget](int32 ThreadID)
 	{
-		int32 VertexID = VerticesToUpdate
+		int32 SubVID = VerticesToUpdate
 			? (*VerticesToUpdate)[ThreadID]
 			: ThreadID;
 
-		if (!MaskToolCommitMesh.IsVertex(VertexID))
+		// Remap submesh VID to full mesh VID for MaskToolCommitMesh access
+		int32 FullVID = IsolationSubmesh ? IsolationSubmesh->MapVertexToBaseMesh(SubVID) : SubVID;
+
+		if (FullVID == FDynamicMesh3::InvalidID || !MaskToolCommitMesh.IsVertex(FullVID))
 		{
 			return;
 		}
 
-		const float Weight = VertexData.GetValue(VertexID);
+		const float Weight = VertexData.GetValue(SubVID);
 
 		FVector Delta;
-		BaseToolMorphTarget->GetValue(VertexID, Delta);
+		BaseToolMorphTarget->GetValue(FullVID, Delta);
 
 		Delta *= (1.0f - Weight);
-		ToolMorphTargetCache->SetValue(VertexID, Delta);
+		ToolMorphTargetCache->SetValue(FullVID, Delta);
 	});
 
 	const TArray<FTransform>& ComponentSpaceTransforms = GetComponentSpaceBoneTransforms();
@@ -631,6 +667,10 @@ GetBaseMorphTarget());
 		}
 	},
 	MaskToolProperties->VertexAttributeOption);
+	
+	// Hard to do real change detection for masking since 
+	// it requires iterating the mesh verts 
+	bAnyChangeMade = true;
 }
 
 
@@ -701,12 +741,49 @@ void UMorphTargetMaskTool::PosePreviewMesh(const TArray<FTransform>& ComponentSp
 
 	FDynamicMesh3& SculptMesh = *GetSculptMesh();
 
+	// When isolated, pose only the isolated vertices and remap IDs to submesh space.
+	// GetPosedMesh operates on MaskToolCommitMesh (full mesh), so we pass full-mesh vertex IDs.
+	TArray<int32> FullMeshVertArrayForIsolation;
+	if (IsolationSubmesh)
+	{
+		if (VerticesToUpdate)
+		{
+			FullMeshVertArrayForIsolation.Reserve(VerticesToUpdate->Num());
+			for (int32 SubVID : *VerticesToUpdate)
+			{
+				FullMeshVertArrayForIsolation.Add(IsolationSubmesh->MapVertexToBaseMesh(SubVID));
+			}
+		}
+		else
+		{
+			FullMeshVertArrayForIsolation.Reserve(IsolationSubmesh->GetSubmesh().VertexCount());
+			for (int32 SubVID : IsolationSubmesh->GetSubmesh().VertexIndicesItr())
+			{
+				FullMeshVertArrayForIsolation.Add(IsolationSubmesh->MapVertexToBaseMesh(SubVID));
+			}
+		}
+	}
+	
+	const TArray<int32>* FinalFullMeshVertArrayPtr = nullptr;
+	if (IsolationSubmesh)
+	{
+		FinalFullMeshVertArrayPtr = &FullMeshVertArrayForIsolation;
+	}
+	else
+	{
+		FinalFullMeshVertArrayPtr = VerticesToUpdate;
+	}
+
 	auto WriteFunc = [&](FVertInfo VertInfo, const FVector& PosedVertPos)
 		{
-			SculptMesh.SetVertex(VertInfo.VertID, PosedVertPos);
+			int32 SculptVID = IsolationSubmesh ? IsolationSubmesh->MapVertexToSubmesh(VertInfo.VertID) : VertInfo.VertID;
+			if (SculptVID != FDynamicMesh3::InvalidID)
+			{
+				SculptMesh.SetVertex(SculptVID, PosedVertPos);
+			}
 		};
 
-	GetPosedMesh(WriteFunc, MaskToolCommitMesh, BoneMatrices, NAME_None, MorphTargetWeights, VerticesToUpdate ? *VerticesToUpdate : TArray<int32>());
+	GetPosedMesh(WriteFunc, MaskToolCommitMesh, BoneMatrices, NAME_None, MorphTargetWeights, FinalFullMeshVertArrayPtr ? *FinalFullMeshVertArrayPtr : TArray<int32>());
 	if (!TrianglesToUpdate)
 	{
 		FMeshNormals::QuickRecomputeOverlayNormals(SculptMesh);
diff --git a/Engine/Plugins/Experimental/Animation/SkeletalMeshMorphTargetEditingTools/Source/SkeletalMeshMorphTargetEditingTools/Public/MorphTargetMaskTool.h b/Engine/Plugins/Experimental/Animation/SkeletalMeshMorphTargetEditingTools/Source/SkeletalMeshMorphTargetEditingTools/Public/MorphTargetMaskTool.h
--- a/Engine/Plugins/Experimental/Animation/SkeletalMeshMorphTargetEditingTools/Source/SkeletalMeshMorphTargetEditingTools/Public/MorphTargetMaskTool.h
+++ b/Engine/Plugins/Experimental/Animation/SkeletalMeshMorphTargetEditingTools/Source/SkeletalMeshMorphTargetEditingTools/Public/MorphTargetMaskTool.h
@@ -4,6 +4,7 @@
 
 #include "IMorphTargetEditingToolInterface.h"
 
+#include "DynamicSubmesh3.h"
 #include "MeshVertexAttributePaintToolBase.h"
 #include "SkeletalMesh/SkeletalMeshToolsHelper.h"
 #include "SkeletalMesh/SkeletalMeshEditingInterface.h"
@@ -122,6 +123,9 @@ protected:
 
 	FDynamicMesh3 MaskToolCommitMesh;
 
+	// Geometry isolation: submesh mapping for vertex ID translation (null when no isolation)
+	TOptional<UE::Geometry::FDynamicSubmesh3> IsolationSubmesh;
+
 	TArray<FTransform> ComponentSpaceTransformsRefPose;
 	TMap<FName, float> MorphTargetZeroWeights;
 	
diff --git a/Engine/Plugins/Experimental/MeshModelingToolsetExp/Source/MeshModelingToolsEditorOnlyExp/Private/MeshVertexAttributePaintToolBase.cpp b/Engine/Plugins/Experimental/MeshModelingToolsetExp/Source/MeshModelingToolsEditorOnlyExp/Private/MeshVertexAttributePaintToolBase.cpp
--- a/Engine/Plugins/Experimental/MeshModelingToolsetExp/Source/MeshModelingToolsEditorOnlyExp/Private/MeshVertexAttributePaintToolBase.cpp
+++ b/Engine/Plugins/Experimental/MeshModelingToolsetExp/Source/MeshModelingToolsEditorOnlyExp/Private/MeshVertexAttributePaintToolBase.cpp
@@ -457,12 +457,14 @@ void UMeshVertexAttributePaintToolBase::Setup()
 	UE::ToolTarget::HideSourceObject(Target);
 	
 	// make sure we get a copy of the original dynamic mesh 
-	FGetMeshParameters GetMeshParams;
-	GetMeshParams.bWantMeshTangents = true;
-	FDynamicMesh3 DynamicMeshCopy = UE::ToolTarget::GetDynamicMeshCopy(Target, GetMeshParams);
+	{
+		FGetMeshParameters GetMeshParams;
+		GetMeshParams.bWantMeshTangents = true;
+		FDynamicMesh3 DynamicMeshCopy = UE::ToolTarget::GetDynamicMeshCopy(Target, GetMeshParams);
 	
-	ToolDynamicMesh = NewObject<UDynamicMesh>();
-	ToolDynamicMesh->SetMesh(DynamicMeshCopy);
+		ToolDynamicMesh = NewObject<UDynamicMesh>();
+		ToolDynamicMesh->SetMesh(MoveTemp(DynamicMeshCopy));
+	}
 
 	ToolDynamicMesh->EditMesh([this](FDynamicMesh3& InMesh)
 	{
@@ -479,7 +481,10 @@ void UMeshVertexAttributePaintToolBase::Setup()
 	PreviewMesh->CreateInWorld(TargetWorld, FTransform::Identity);
 	PreviewMesh->SetTangentsMode(EDynamicMeshComponentTangentsMode::AutoCalculated);
 	ToolSetupUtil::ApplyRenderingConfigurationToPreview(PreviewMesh, Target);
-	PreviewMesh->ReplaceMesh(MoveTemp(DynamicMeshCopy));
+	ToolDynamicMesh->ProcessMesh([this](const FDynamicMesh3& InToolMesh)
+	{
+		PreviewMesh->ReplaceMesh(InToolMesh);
+	});
 		
 	// assign materials
 	FComponentMaterialSet MaterialSet = UE::ToolTarget::GetMaterialSet(Target);
diff --git a/Engine/Plugins/Runtime/MeshModelingToolset/Source/MeshModelingToolsEditorOnly/Public/SkeletalMesh/SkeletalMeshEditingInterface.h b/Engine/Plugins/Runtime/MeshModelingToolset/Source/MeshModelingToolsEditorOnly/Public/SkeletalMesh/SkeletalMeshEditingInterface.h
--- a/Engine/Plugins/Runtime/MeshModelingToolset/Source/MeshModelingToolsEditorOnly/Public/SkeletalMesh/SkeletalMeshEditingInterface.h
+++ b/Engine/Plugins/Runtime/MeshModelingToolset/Source/MeshModelingToolsEditorOnly/Public/SkeletalMesh/SkeletalMeshEditingInterface.h
@@ -160,6 +160,8 @@ public:
 	virtual TMap<FName, float> GetMorphTargetWeights() PURE_VIRTUAL(USkeletalMeshEditorContextObjectBase::GetMorphTargetWeights, return {};);	
 	virtual void NotifyMorphTargetEdited() PURE_VIRTUAL(USkeletalMeshEditorContextObjectBase::NotifyMorphTargetEdited, return;);
 	
+	virtual const TArray<int32>& GetIsolatedTriangles() const { static const TArray<int32> Empty; return Empty; }
+	
 	virtual void BindTo(ISkeletalMeshEditingInterface* InEditingInterface) PURE_VIRTUAL(USkeletalMeshEditorContextObjectBase::BindTo, return;)
 	virtual void UnbindFrom(ISkeletalMeshEditingInterface* InEditingInterface) PURE_VIRTUAL(USkeletalMeshEditorContextObjectBase::UnbindFrom, return;)
 };
