[SkeletalMeshModelingTools][MorphTargetSculptTool] Add geometry isolation support

#Synced-CL 52717952

diff --git a/Engine/Plugins/Experimental/Animation/SkeletalMeshMorphTargetEditingTools/Source/SkeletalMeshMorphTargetEditingTools/Private/MorphTargetVertexSculptTool.cpp b/Engine/Plugins/Experimental/Animation/SkeletalMeshMorphTargetEditingTools/Source/SkeletalMeshMorphTargetEditingTools/Private/MorphTargetVertexSculptTool.cpp
--- a/Engine/Plugins/Experimental/Animation/SkeletalMeshMorphTargetEditingTools/Source/SkeletalMeshMorphTargetEditingTools/Private/MorphTargetVertexSculptTool.cpp
+++ b/Engine/Plugins/Experimental/Animation/SkeletalMeshMorphTargetEditingTools/Source/SkeletalMeshMorphTargetEditingTools/Private/MorphTargetVertexSculptTool.cpp
@@ -30,6 +30,7 @@
 #include "TargetInterfaces/SkeletonProvider.h"
 #include "DynamicMesh/MeshNormals.h"
 #include "Sculpting/MeshSmoothingBrushOps.h"
+#include "ToolSetupUtil.h"
 
 #include UE_INLINE_GENERATED_CPP_BY_NAME(MorphTargetVertexSculptTool)
 
@@ -137,6 +138,53 @@ void UMorphTargetVertexSculptTool::HandlePoseChangeDetectorEvent(SkeletalMeshToo
 
 
 
+void UMorphTargetVertexSculptTool::InitializeSculptMeshComponent(UBaseDynamicMeshComponent* Component, AActor* Actor)
+{
+	ToolSetupUtil::ApplyRenderingConfigurationToPreview(Component, nullptr);
+	Component->SetShadowsEnabled(false);
+	Component->SetupAttachment(Actor->GetRootComponent());
+	Component->RegisterComponent();
+
+	static FGetMeshParameters GetMeshParams;
+	GetMeshParams.bWantMeshTangents = true;
+	FDynamicMesh3 FullMesh = UE::ToolTarget::GetDynamicMeshCopy(Target, GetMeshParams);
+
+	ToolDynamicMesh = NewObject<UDynamicMesh>();
+	ToolDynamicMesh->SetMesh(MoveTemp(FullMesh));
+
+	// Build isolated submesh if isolation was active, set the correct mesh on the preview/sculpt component
+	ToolDynamicMesh->ProcessMesh([&](const FDynamicMesh3& Mesh)
+		{
+			const FDynamicMesh3* MeshForPreview = &Mesh;
+			
+			if (EditorContext.IsValid())
+			{
+				const TArray<int32>& IsolatedTriangles = EditorContext->GetIsolatedTriangles();
+				if (!IsolatedTriangles.IsEmpty())
+				{
+					IsolationSubmesh = UE::Geometry::FDynamicSubmesh3(&Mesh, IsolatedTriangles);
+					MeshForPreview = &IsolationSubmesh->GetSubmesh();
+				}
+			}
+
+			Component->SetMesh(FDynamicMesh3(*MeshForPreview));
+		});
+
+
+	// Bake rotation and scaling into the component's preview mesh only so that we don't have to deal with skewed brush stamps
+	// TODO: for this to work for morph targets we will need to unbake the transforms before extracting the morph delta, so comment out this part for now
+	
+	// InitialTargetTransform = UE::ToolTarget::GetLocalToWorldTransform(Target);
+	// InitialTargetTransform.ClampMinimumScale(0.01);
+	// FVector3d Translation = InitialTargetTransform.GetTranslation();
+	// InitialTargetTransform.SetTranslation(FVector3d::Zero());
+	// Component->ApplyTransform(InitialTargetTransform, false);
+	// CurTargetTransform = UE::Geometry::FTransformSRT3d(Translation);
+	// Component->SetWorldTransform((FTransform)CurTargetTransform);
+
+	UE::ToolTarget::HideSourceObject(Target);
+}
+
 void UMorphTargetVertexSculptTool::Setup()
 {
 	EditorContext = GetToolManager()->GetContextObjectStore()->FindContext<USkeletalMeshEditorContextObjectBase>();
@@ -148,18 +196,14 @@ void UMorphTargetVertexSculptTool::Setup()
 
 	ViewProperties->MaterialMode = EMeshEditingMaterialModes::ExistingMaterial;
 
-	ToolDynamicMesh = NewObject<UDynamicMesh>();
-
-	ToolDynamicMesh->SetMesh(*GetSculptMesh());
-
-	
+	// ToolDynamicMesh created in InitializeSculptMeshComponent
 	
 	// Remove unused attributes to make undo/redo faster
 	GetSculptMesh()->Attributes()->RemoveAllMorphTargetAttributes();
 	GetSculptMesh()->Attributes()->RemoveAllSkinWeightsAttributes();
 	
 	constexpr bool bCopyNormals = false, bCopyColors = false, bCopyUVs = false, bCopyAttributes = false;
-	MeshWithoutEditingMorph.Copy(*GetSculptMesh(), bCopyNormals, bCopyColors, bCopyUVs, bCopyAttributes);
+	PosedMeshWithoutEditingMorph.Copy(*GetSculptMesh(), bCopyNormals, bCopyColors, bCopyUVs, bCopyAttributes);
 	
 	ToolMorphTargetName = GetEditingMorphTarget();
 	
@@ -177,11 +221,11 @@ void UMorphTargetVertexSculptTool::RegisterBrushes()
 	
 	GetMeshWithoutCurrentMorphFunc = [this]()
 	{
-		if (bMeshWithoutEditingMorphDirty)
+		if (bPosedMeshWithoutEditingMorphDirty)
 		{
-			UpdateMeshWithoutEditingMorph();
+			UpdatePosedMeshWithoutEditingMorph();
 		}
-		return &MeshWithoutEditingMorph;
+		return &PosedMeshWithoutEditingMorph;
 	};
 
 	// Had to hijack the EraseSculptLayer identifier from base mesh vertex sculpt tool for our erase morph target tool since it is the simplest way to get an icon for the tool.
@@ -362,21 +406,44 @@ void UMorphTargetVertexSculptTool::OnEndStroke()
 				{
 					FDynamicMeshMorphTargetAttribute* MorphTargetAttribute = Mesh.Attributes()->GetMorphTargetAttribute(ToolMorphTargetName);
 
+					// When isolated, use submesh as SourceMesh (matching PosedMesh vertex ID space).
+					// Remap to full mesh VIDs in WriteFunction for accessing ToolDynamicMesh.
+					const FDynamicMesh3& SourceMeshForUnpose = IsolationSubmesh
+						? IsolationSubmesh->GetSubmesh()
+						: Mesh;
+
 					auto WriteFunction = [&](FVertInfo VertInfo, const FVector& UnposedVertPos)
 						{
-							int32 VertID = VertInfo.VertID;
-							FVector NewDelta = UnposedVertPos - FVector(Mesh.GetVertex(VertID));
+							int32 FullVID = IsolationSubmesh
+								? IsolationSubmesh->MapVertexToBaseMesh(VertInfo.VertID)
+								: VertInfo.VertID;
+
+							if (FullVID == FDynamicMesh3::InvalidID || !Mesh.IsVertex(FullVID))
+							{
+								return;
+							}
+
+							FVector NewDelta = UnposedVertPos - FVector(Mesh.GetVertex(FullVID));
 							NewDelta /= EditingMorphTargetWeight;
 
 							FVector OldDelta;
-							MorphTargetAttribute->GetValue(VertID, OldDelta);
-							MorphTargetAttribute->SetValue(VertID, NewDelta);
+							MorphTargetAttribute->GetValue(FullVID, OldDelta);
+							MorphTargetAttribute->SetValue(FullVID, NewDelta);
 
 							MorphTargetChange->OldDeltas[VertInfo.VertArrayIndex] = OldDelta;
 							MorphTargetChange->NewDeltas[VertInfo.VertArrayIndex] = NewDelta;
 						};
 
-					GetUnposedMesh(WriteFunction, *GetSculptMesh(), Mesh, BoneMatrices, NAME_None, MorphTargetWeights, MorphTargetChange->Vertices);
+					GetUnposedMesh(WriteFunction, *GetSculptMesh(), SourceMeshForUnpose, BoneMatrices, NAME_None, MorphTargetWeights, MorphTargetChange->Vertices);
+
+					// Remap MorphTargetChange vertices from submesh to full mesh for undo/redo
+					if (IsolationSubmesh)
+					{
+						for (int32& VID : MorphTargetChange->Vertices)
+						{
+							VID = IsolationSubmesh->MapVertexToBaseMesh(VID);
+						}
+					}
 				});
 
 
@@ -446,28 +513,42 @@ void UMorphTargetVertexSculptTool::PoseSculptMesh(const TArray<FTransform>& Comp
 
 	ToolDynamicMesh->ProcessMesh([&](const FDynamicMesh3& InMesh)
 		{
+			// When isolated, build full-mesh vertex array and remap in WriteFunc
+			TArray<int32> FullMeshVertArray;
+			if (IsolationSubmesh)
+			{
+				FullMeshVertArray.Reserve(IsolationSubmesh->GetSubmesh().VertexCount());
+				for (int32 SubVID : IsolationSubmesh->GetSubmesh().VertexIndicesItr())
+				{
+					FullMeshVertArray.Add(IsolationSubmesh->MapVertexToBaseMesh(SubVID));
+				}
+			}
+
 			auto WriteFunc = [&](FVertInfo VertInfo, const FVector& PosedVertPos)
 				{
-					SculptMesh.SetVertex(VertInfo.VertID, PosedVertPos);
+					int32 SculptVID = IsolationSubmesh ? IsolationSubmesh->MapVertexToSubmesh(VertInfo.VertID) : VertInfo.VertID;
+					if (SculptVID != FDynamicMesh3::InvalidID)
+					{
+						SculptMesh.SetVertex(SculptVID, PosedVertPos);
+					}
 				};
 
-			GetPosedMesh(WriteFunc, InMesh, BoneMatrices, NAME_None, MorphTargetWeights);
+			GetPosedMesh(WriteFunc, InMesh, BoneMatrices, NAME_None, MorphTargetWeights,
+				IsolationSubmesh ? FullMeshVertArray : TArray<int32>());
 			FMeshNormals::QuickRecomputeOverlayNormals(SculptMesh);
 			
 			ComponentSpaceTransformsForPosedMesh = ComponentSpaceTransforms;
 			MorphTargetWeightsForPosedMesh = MorphTargetWeights;
 		});
 
-	bMeshWithoutEditingMorphDirty = true;
+	bPosedMeshWithoutEditingMorphDirty = true;
 
 	constexpr bool bNormal = true;
 	DynamicMeshComponent->FastNotifyPositionsUpdated(bNormal);
-
-
 }
 
 
-void UMorphTargetVertexSculptTool::UpdateMeshWithoutEditingMorph()
+void UMorphTargetVertexSculptTool::UpdatePosedMeshWithoutEditingMorph()
 {
 	using namespace SkeletalMeshToolsHelper;
 
@@ -482,15 +563,31 @@ void UMorphTargetVertexSculptTool::UpdateMeshWithoutEditingMorph()
 
 	ToolDynamicMesh->ProcessMesh([&](const FDynamicMesh3& InMesh)
 		{
+			// When isolated, build full-mesh vertex array and remap in WriteFunc
+			TArray<int32> FullMeshVertArray;
+			if (IsolationSubmesh)
+			{
+				FullMeshVertArray.Reserve(IsolationSubmesh->GetSubmesh().VertexCount());
+				for (int32 SubVID : IsolationSubmesh->GetSubmesh().VertexIndicesItr())
+				{
+					FullMeshVertArray.Add(IsolationSubmesh->MapVertexToBaseMesh(SubVID));
+				}
+			}
+
 			auto WriteFunc = [&](FVertInfo VertInfo, const FVector& PosedVertPos)
 				{
-					MeshWithoutEditingMorph.SetVertex(VertInfo.VertID, PosedVertPos);
+					int32 WriteVID = IsolationSubmesh ? IsolationSubmesh->MapVertexToSubmesh(VertInfo.VertID) : VertInfo.VertID;
+					if (WriteVID != FDynamicMesh3::InvalidID)
+					{
+						PosedMeshWithoutEditingMorph.SetVertex(WriteVID, PosedVertPos);
+					}
 				};
 
-			GetPosedMesh(WriteFunc, InMesh, BoneMatrices, NAME_None, WeightsWithoutEditingMorph);
+			GetPosedMesh(WriteFunc, InMesh, BoneMatrices, NAME_None, WeightsWithoutEditingMorph,
+				IsolationSubmesh ? FullMeshVertArray : TArray<int32>());
 		});
 
-	bMeshWithoutEditingMorphDirty = false;
+	bPosedMeshWithoutEditingMorphDirty = false;
 }
 
 #undef LOCTEXT_NAMESPACE
diff --git a/Engine/Plugins/Experimental/Animation/SkeletalMeshMorphTargetEditingTools/Source/SkeletalMeshMorphTargetEditingTools/Public/MorphTargetVertexSculptTool.h b/Engine/Plugins/Experimental/Animation/SkeletalMeshMorphTargetEditingTools/Source/SkeletalMeshMorphTargetEditingTools/Public/MorphTargetVertexSculptTool.h
--- a/Engine/Plugins/Experimental/Animation/SkeletalMeshMorphTargetEditingTools/Source/SkeletalMeshMorphTargetEditingTools/Public/MorphTargetVertexSculptTool.h
+++ b/Engine/Plugins/Experimental/Animation/SkeletalMeshMorphTargetEditingTools/Source/SkeletalMeshMorphTargetEditingTools/Public/MorphTargetVertexSculptTool.h
@@ -4,6 +4,7 @@
 
 #include "IMorphTargetEditingToolInterface.h"
 #include "MeshVertexSculptTool.h"
+#include "DynamicSubmesh3.h"
 #include "SkeletalMesh/SkeletalMeshEditingInterface.h"
 #include "BaseTools/MultiSelectionMeshEditingTool.h"
 #include "MeshDescription.h"
@@ -47,6 +48,9 @@ public:
 	virtual void ToggleBoneManipulation(bool bEnable);
 
 
+	// UMeshSculptToolBase overrides
+	virtual void InitializeSculptMeshComponent(UBaseDynamicMeshComponent* Component, AActor* Actor) override;
+
 	// UMeshVertexSculptTool overrides
 	virtual void Setup() override;
 	virtual void Shutdown(EToolShutdownType ShutdownType) override;
@@ -82,9 +86,9 @@ protected:
 	TWeakObjectPtr<USkeletalMeshEditorContextObjectBase> EditorContext = nullptr;
 
 	TFunction<const FDynamicMesh3*()> GetMeshWithoutCurrentMorphFunc;
-	FDynamicMesh3 MeshWithoutEditingMorph;
-	bool bMeshWithoutEditingMorphDirty = true;
-	void UpdateMeshWithoutEditingMorph();
+	FDynamicMesh3 PosedMeshWithoutEditingMorph;
+	bool bPosedMeshWithoutEditingMorphDirty = true;
+	void UpdatePosedMeshWithoutEditingMorph();
 	
 	SkeletalMeshToolsHelper::FPoseChangeDetector PoseChangeDetector;
 
@@ -94,6 +98,9 @@ protected:
 
 	UPROPERTY()
 	TObjectPtr<UDynamicMesh> ToolDynamicMesh;
+
+	// Geometry isolation: submesh mapping for vertex ID translation (null when no isolation)
+	TOptional<UE::Geometry::FDynamicSubmesh3> IsolationSubmesh;
 	
 	TArray<FTransform> ComponentSpaceTransformsRefPose;
 	
diff --git a/Engine/Plugins/Runtime/MeshModelingToolset/Source/MeshModelingTools/Public/Sculpting/MeshSculptToolBase.h b/Engine/Plugins/Runtime/MeshModelingToolset/Source/MeshModelingTools/Public/Sculpting/MeshSculptToolBase.h
--- a/Engine/Plugins/Runtime/MeshModelingToolset/Source/MeshModelingTools/Public/Sculpting/MeshSculptToolBase.h
+++ b/Engine/Plugins/Runtime/MeshModelingToolset/Source/MeshModelingTools/Public/Sculpting/MeshSculptToolBase.h
@@ -330,7 +330,7 @@ protected:
 	/**
 	 * Subclass calls this to set up editing component
 	 */
-	UE_API void InitializeSculptMeshComponent(UBaseDynamicMeshComponent* Component, AActor* Actor);
+	UE_API virtual void InitializeSculptMeshComponent(UBaseDynamicMeshComponent* Component, AActor* Actor);
 
 
 	/**
